ChaosOutfitAsset - Added outfit asset stubs.

#rnx

#virtualized

diff --git a/Engine/Plugins/Experimental/ChaosOutfitAsset/ChaosOutfitAsset.uplugin b/Engine/Plugins/Experimental/ChaosOutfitAsset/ChaosOutfitAsset.uplugin
new file mode 100644
--- /dev/null
+++ b/Engine/Plugins/Experimental/ChaosOutfitAsset/ChaosOutfitAsset.uplugin
@@ -0,0 +1,46 @@
+{
+	"FileVersion": 3,
+	"Version": 1,
+	"VersionName": "0.1",
+	"FriendlyName": "Chaos Outfit Asset",
+	"Description": "Outfit Asset plugin to create and assemble outfits made of Cloth Assets.",
+	"Category": "Physics",
+	"CreatedBy": "Epic Games, Inc.",
+	"CreatedByURL": "https://epicgames.com",
+	"DocsURL": "",
+	"MarketplaceURL": "",
+	"SupportURL": "",
+	"EnabledByDefault": false,
+	"CanContainContent": true,
+	"IsBetaVersion": false,
+	"IsExperimentalVersion": true,
+	"Installed": false,
+	"Modules": [
+		{
+			"Name": "ChaosOutfitAssetDataflowNodes",
+			"Type": "Runtime",
+			"LoadingPhase": "Default"
+		},
+		{
+			"Name": "ChaosOutfitAssetEngine",
+			"Type": "Runtime",
+			"LoadingPhase": "Default"
+		},
+		{
+			"Name": "ChaosOutfitAssetEditor",
+			"Type": "Editor",
+			"LoadingPhase": "Default",
+			"PlatformAllowList": [ "Win64", "Mac", "Linux" ]
+		}
+	],
+	"Plugins": [
+		{
+			"Name": "ChaosClothAsset",
+			"Enabled": true
+		},
+		{
+			"Name": "Dataflow",
+			"Enabled": true
+		}
+	]
+}
diff --git a/Engine/Plugins/Experimental/ChaosOutfitAsset/Source/ChaosOutfitAssetDataflowNodes/ChaosOutfitAssetDataflowNodes.Build.cs b/Engine/Plugins/Experimental/ChaosOutfitAsset/Source/ChaosOutfitAssetDataflowNodes/ChaosOutfitAssetDataflowNodes.Build.cs
new file mode 100644
--- /dev/null
+++ b/Engine/Plugins/Experimental/ChaosOutfitAsset/Source/ChaosOutfitAssetDataflowNodes/ChaosOutfitAssetDataflowNodes.Build.cs
@@ -0,0 +1,24 @@
+// Copyright Epic Games, Inc. All Rights Reserved.
+
+using UnrealBuildTool;
+
+public class ChaosOutfitAssetDataflowNodes : ModuleRules
+{
+	public ChaosOutfitAssetDataflowNodes(ReadOnlyTargetRules Target) : base(Target)
+	{
+		PCHUsage = ModuleRules.PCHUsageMode.UseExplicitOrSharedPCHs;
+
+		PrivateDependencyModuleNames.AddRange(
+			new string[]
+			{
+				"Core",
+				"CoreUObject",
+				"Chaos",
+				"ChaosClothAssetEngine",
+				"ChaosOutfitAssetEngine",
+				"DataflowCore",
+				"DataflowNodes",
+			}
+		);
+	}
+}
diff --git a/Engine/Plugins/Experimental/ChaosOutfitAsset/Source/ChaosOutfitAssetDataflowNodes/Private/ChaosClothAsset/ChaosOutfitAssetDataflowNodesModule.cpp b/Engine/Plugins/Experimental/ChaosOutfitAsset/Source/ChaosOutfitAssetDataflowNodes/Private/ChaosClothAsset/ChaosOutfitAssetDataflowNodesModule.cpp
new file mode 100644
--- /dev/null
+++ b/Engine/Plugins/Experimental/ChaosOutfitAsset/Source/ChaosOutfitAssetDataflowNodes/Private/ChaosClothAsset/ChaosOutfitAssetDataflowNodesModule.cpp
@@ -0,0 +1,48 @@
+// Copyright Epic Games, Inc. All Rights Reserved.
+
+#include "ChaosClothAsset/ClothAssetToOutfitCollectionNode.h"
+#include "ChaosClothAsset/GetClothAssetNode.h"
+#include "ChaosClothAsset/OutfitAssetTerminalNode.h"
+#include "Dataflow/DataflowNodeColorsRegistry.h"
+#include "Dataflow/DataflowNodeFactory.h"
+#include "Modules/ModuleInterface.h"
+#include "Modules/ModuleManager.h"
+
+#define LOCTEXT_NAMESPACE "OutfitAssetDataflowNodesModule"
+
+namespace UE::Chaos::ClothAsset
+{
+	namespace Private
+	{
+		static const FLinearColor OutfitAssetNodeHeaderColor = FColor(162, 108, 99);
+		static const FLinearColor OutfitAssetNodeBodyColor = FColor(18, 12, 11, 127);
+
+		static void RegisterDataflowNodes()
+		{
+			DATAFLOW_NODE_REGISTER_CREATION_FACTORY_NODE_COLORS_BY_CATEGORY("Outfit", OutfitAssetNodeHeaderColor, OutfitAssetNodeBodyColor);
+			DATAFLOW_NODE_REGISTER_CREATION_FACTORY(FChaosClothAssetToOutfitCollectionNode);
+			DATAFLOW_NODE_REGISTER_CREATION_FACTORY(FChaosGetClothAssetNode);
+			DATAFLOW_NODE_REGISTER_CREATION_FACTORY(FChaosOutfitAssetTerminalNode);
+PRAGMA_DISABLE_DEPRECATION_WARNINGS
+			// Move any deprecated node registrations under here
+PRAGMA_ENABLE_DEPRECATION_WARNINGS
+		}
+	}  // namespace Private
+
+	class FOutfitAssetDataflowNodesModule : public IModuleInterface
+	{
+	public:
+		virtual void StartupModule() override
+		{
+			Private::RegisterDataflowNodes();
+		}
+
+		virtual void ShutdownModule() override
+		{
+		}
+	};
+}  // namespace UE::Chaos::ClothAsset
+
+IMPLEMENT_MODULE(UE::Chaos::ClothAsset::FOutfitAssetDataflowNodesModule, ChaosOutfitAssetDataflowNodes);
+
+#undef LOCTEXT_NAMESPACE
diff --git a/Engine/Plugins/Experimental/ChaosOutfitAsset/Source/ChaosOutfitAssetDataflowNodes/Private/ChaosClothAsset/ClothAssetToOutfitCollectionNode.cpp b/Engine/Plugins/Experimental/ChaosOutfitAsset/Source/ChaosOutfitAssetDataflowNodes/Private/ChaosClothAsset/ClothAssetToOutfitCollectionNode.cpp
new file mode 100644
--- /dev/null
+++ b/Engine/Plugins/Experimental/ChaosOutfitAsset/Source/ChaosOutfitAssetDataflowNodes/Private/ChaosClothAsset/ClothAssetToOutfitCollectionNode.cpp
@@ -0,0 +1,113 @@
+// Copyright Epic Games, Inc. All Rights Reserved.
+
+#include "ChaosClothAsset/ClothAssetToOutfitCollectionNode.h"
+
+#include UE_INLINE_GENERATED_CPP_BY_NAME(ClothAssetToOutfitCollectionNode)
+
+#define LOCTEXT_NAMESPACE "ChaosClothAssetToOutfitCollectionNode"
+
+FChaosClothAssetToOutfitCollectionNode::FChaosClothAssetToOutfitCollectionNode(const UE::Dataflow::FNodeParameters& InParam, FGuid InGuid)
+	: FDataflowNode(InParam, InGuid)
+{
+	RegisterInputConnection(&Collection);
+	RegisterOutputConnection(&Collection, &Collection);
+
+	for (int32 Index = 0; Index < NumInitialClothAssets; ++Index)
+	{
+		AddPins();
+	}
+	check(GetNumInputs() == NumRequiredInputs + NumInitialClothAssets);  // Update NumRequiredInputs when adding inputs (used by Serialize)
+}
+
+void FChaosClothAssetToOutfitCollectionNode::Evaluate(UE::Dataflow::FContext& Context, const FDataflowOutput* Out) const
+{
+	if (Out->IsA<FManagedArrayCollection>(&Collection))
+	{
+		const FManagedArrayCollection& OutfitCollection = GetValue(Context, &Collection);
+		// TODO: Add assets into the OutfitCollection using an outfit facade
+
+		SetValue(Context, OutfitCollection, &Collection);
+	}
+}
+
+UE::Dataflow::TConnectionReference<TObjectPtr<const UChaosClothAsset>> FChaosClothAssetToOutfitCollectionNode::GetConnectionReference(int32 Index) const
+{
+	return { &ClothAssets[Index], Index, &ClothAssets };
+}
+
+TArray<UE::Dataflow::FPin> FChaosClothAssetToOutfitCollectionNode::AddPins()
+{
+	const int32 Index = ClothAssets.AddDefaulted();
+	const FDataflowInput& Input = RegisterInputArrayConnection(GetConnectionReference(Index));
+	return { { UE::Dataflow::FPin::EDirection::INPUT, Input.GetType(), Input.GetName() } };
+}
+
+TArray<UE::Dataflow::FPin> FChaosClothAssetToOutfitCollectionNode::GetPinsToRemove() const
+{
+	const int32 Index = ClothAssets.Num() - 1;
+	check(ClothAssets.IsValidIndex(Index));
+	if (const FDataflowInput* const Input = FindInput(GetConnectionReference(Index)))
+	{
+		return { { UE::Dataflow::FPin::EDirection::INPUT, Input->GetType(), Input->GetName() } };
+	}
+	return Super::GetPinsToRemove();
+}
+
+void FChaosClothAssetToOutfitCollectionNode::OnPinRemoved(const UE::Dataflow::FPin& Pin)
+{
+	const int32 Index = ClothAssets.Num() - 1;
+	check(ClothAssets.IsValidIndex(Index));
+#if DO_CHECK
+	const FDataflowInput* const Input = FindInput(GetConnectionReference(Index));
+	check(Input);
+	check(Input->GetName() == Pin.Name);
+	check(Input->GetType() == Pin.Type);
+#endif
+	ClothAssets.SetNum(Index);
+
+	return Super::OnPinRemoved(Pin);
+}
+
+void FChaosClothAssetToOutfitCollectionNode::PostSerialize(const FArchive& Ar)
+{
+	// Added pins need to be restored when loading to make sure they get reconnected
+	if (Ar.IsLoading())
+	{
+		if (ClothAssets.Num() < NumInitialClothAssets)
+		{
+			ClothAssets.SetNum(NumInitialClothAssets);  // In case the FManagedArrayCollection wasn't serialized with the node (pre the WithSerializer trait)
+		}
+		for (int32 Index = 0; Index < NumInitialClothAssets; ++Index)
+		{
+			check(FindInput(GetConnectionReference(Index)));
+		}
+
+		for (int32 Index = NumInitialClothAssets; Index < ClothAssets.Num(); ++Index)
+		{
+			FindOrRegisterInputArrayConnection(GetConnectionReference(Index));
+		}
+		if (Ar.IsTransacting())
+		{
+			const int32 OrigNumRegisteredInputs = GetNumInputs();
+			check(OrigNumRegisteredInputs >= NumRequiredInputs + NumInitialClothAssets);
+			const int32 OrigNumClothAssets = ClothAssets.Num();
+			const int32 OrigNumRegisteredClothAssets = (OrigNumRegisteredInputs - NumRequiredInputs);
+			if (OrigNumRegisteredClothAssets > OrigNumClothAssets)
+			{
+				// Inputs have been removed, temporarily expand ClothAssets so we can get connection references
+				ClothAssets.SetNum(OrigNumRegisteredClothAssets);
+				for (int32 Index = OrigNumClothAssets; Index < ClothAssets.Num(); ++Index)
+				{
+					UnregisterInputConnection(GetConnectionReference(Index));
+				}
+				ClothAssets.SetNum(OrigNumClothAssets);
+			}
+		}
+		else
+		{
+			ensureAlways(ClothAssets.Num() + NumRequiredInputs == GetNumInputs());
+		}
+	}
+}
+
+#undef LOCTEXT_NAMESPACE
diff --git a/Engine/Plugins/Experimental/ChaosOutfitAsset/Source/ChaosOutfitAssetDataflowNodes/Private/ChaosClothAsset/ClothAssetToOutfitCollectionNode.h b/Engine/Plugins/Experimental/ChaosOutfitAsset/Source/ChaosOutfitAssetDataflowNodes/Private/ChaosClothAsset/ClothAssetToOutfitCollectionNode.h
new file mode 100644
--- /dev/null
+++ b/Engine/Plugins/Experimental/ChaosOutfitAsset/Source/ChaosOutfitAssetDataflowNodes/Private/ChaosClothAsset/ClothAssetToOutfitCollectionNode.h
@@ -0,0 +1,44 @@
+// Copyright Epic Games, Inc. All Rights Reserved.
+
+#pragma once 
+
+#include "Dataflow/DataflowNode.h"
+#include "GeometryCollection/ManagedArrayCollection.h"
+#include "ClothAssetToOutfitCollectionNode.generated.h"
+
+class UChaosClothAsset;
+
+/** Add multiple cloth asset objects to an outfit collection. */
+USTRUCT(Meta = (Experimental, DataflowOutfit))
+struct FChaosClothAssetToOutfitCollectionNode final : public FDataflowNode
+{
+	GENERATED_USTRUCT_BODY()
+	DATAFLOW_NODE_DEFINE_INTERNAL(FChaosClothAssetToOutfitCollectionNode, "ClothAssetToOutfitCollection", "Outfit", "Cloth Asset Outfit Collection")
+
+public:
+	FChaosClothAssetToOutfitCollectionNode(const UE::Dataflow::FNodeParameters& InParam, FGuid InGuid = FGuid::NewGuid());
+
+private:
+	UE::Dataflow::TConnectionReference<TObjectPtr<const UChaosClothAsset>> GetConnectionReference(int32 Index) const;
+
+	//~ Begin FDataflowNode interface
+	virtual void Evaluate(UE::Dataflow::FContext& Context, const FDataflowOutput* Out) const override;
+	virtual TArray<UE::Dataflow::FPin> AddPins() override;
+	virtual bool CanAddPin() const override { return true; }
+	virtual bool CanRemovePin() const override { return ClothAssets.Num() > NumInitialClothAssets; }
+	virtual TArray<UE::Dataflow::FPin> GetPinsToRemove() const override;
+	virtual void OnPinRemoved(const UE::Dataflow::FPin& Pin) override;
+	virtual void PostSerialize(const FArchive& Ar) override;
+	//~ End FDataflowNode interface
+
+	static constexpr int32 NumRequiredInputs = 1;
+	static constexpr int32 NumInitialClothAssets = 1;
+
+	/** The cloth assets to import into the collection. */
+	UPROPERTY(EditAnywhere, Category = "GetClothAsset")
+	TArray<TObjectPtr<const UChaosClothAsset>> ClothAssets;
+
+	/** The outfit output collection. */
+	UPROPERTY(EditAnywhere, Category = "GetClothAsset", Meta = (DataflowInput, DataflowOutput, DataflowPassthrough = "Collection"))
+	FManagedArrayCollection Collection;
+};
diff --git a/Engine/Plugins/Experimental/ChaosOutfitAsset/Source/ChaosOutfitAssetDataflowNodes/Private/ChaosClothAsset/GetClothAssetNode.cpp b/Engine/Plugins/Experimental/ChaosOutfitAsset/Source/ChaosOutfitAssetDataflowNodes/Private/ChaosClothAsset/GetClothAssetNode.cpp
new file mode 100644
--- /dev/null
+++ b/Engine/Plugins/Experimental/ChaosOutfitAsset/Source/ChaosOutfitAssetDataflowNodes/Private/ChaosClothAsset/GetClothAssetNode.cpp
@@ -0,0 +1,24 @@
+// Copyright Epic Games, Inc. All Rights Reserved.
+
+#include "ChaosClothAsset/GetClothAssetNode.h"
+#include "ChaosClothAsset/ClothAsset.h"
+
+#include UE_INLINE_GENERATED_CPP_BY_NAME(GetClothAssetNode)
+
+#define LOCTEXT_NAMESPACE "ChaosGetClothAssetNode"
+
+FChaosGetClothAssetNode::FChaosGetClothAssetNode(const UE::Dataflow::FNodeParameters& InParam, FGuid InGuid)
+	: FDataflowNode(InParam, InGuid)
+{
+	RegisterOutputConnection(&ClothAsset);
+}
+
+void FChaosGetClothAssetNode::Evaluate(UE::Dataflow::FContext& Context, const FDataflowOutput* Out) const
+{
+	if (Out->IsA<TObjectPtr<const UChaosClothAsset>>(&ClothAsset))
+	{
+		SetValue(Context, ClothAsset, &ClothAsset);
+	}
+}
+
+#undef LOCTEXT_NAMESPACE
diff --git a/Engine/Plugins/Experimental/ChaosOutfitAsset/Source/ChaosOutfitAssetDataflowNodes/Private/ChaosClothAsset/GetClothAssetNode.h b/Engine/Plugins/Experimental/ChaosOutfitAsset/Source/ChaosOutfitAssetDataflowNodes/Private/ChaosClothAsset/GetClothAssetNode.h
new file mode 100644
--- /dev/null
+++ b/Engine/Plugins/Experimental/ChaosOutfitAsset/Source/ChaosOutfitAssetDataflowNodes/Private/ChaosClothAsset/GetClothAssetNode.h
@@ -0,0 +1,31 @@
+// Copyright Epic Games, Inc. All Rights Reserved.
+
+#pragma once 
+
+#include "Dataflow/DataflowNode.h"
+#include "GeometryCollection/ManagedArrayCollection.h"
+#include "GetClothAssetNode.generated.h"
+
+class UChaosClothAsset;
+
+/**
+ * Get a cloth asset object into the graph.
+ */
+USTRUCT(Meta = (Experimental, DataflowOutfit))
+struct FChaosGetClothAssetNode final : public FDataflowNode
+{
+	GENERATED_USTRUCT_BODY()
+	DATAFLOW_NODE_DEFINE_INTERNAL(FChaosGetClothAssetNode, "GetClothAsset", "Outfit", "Get Cloth Asset")
+
+public:
+	FChaosGetClothAssetNode(const UE::Dataflow::FNodeParameters& InParam, FGuid InGuid = FGuid::NewGuid());
+
+private:
+	//~ Begin FDataflowNode interface
+	virtual void Evaluate(UE::Dataflow::FContext& Context, const FDataflowOutput* Out) const override;
+	//~ End FDataflowNode interface
+
+	/** The Cloth Asset to import into a collection. */
+	UPROPERTY(EditAnywhere, Category = "GetClothAsset", Meta = (DataflowOutput))
+	TObjectPtr<const UChaosClothAsset> ClothAsset;
+};
diff --git a/Engine/Plugins/Experimental/ChaosOutfitAsset/Source/ChaosOutfitAssetDataflowNodes/Private/ChaosClothAsset/OutfitAssetTerminalNode.cpp b/Engine/Plugins/Experimental/ChaosOutfitAsset/Source/ChaosOutfitAssetDataflowNodes/Private/ChaosClothAsset/OutfitAssetTerminalNode.cpp
new file mode 100644
--- /dev/null
+++ b/Engine/Plugins/Experimental/ChaosOutfitAsset/Source/ChaosOutfitAssetDataflowNodes/Private/ChaosClothAsset/OutfitAssetTerminalNode.cpp
@@ -0,0 +1,32 @@
+// Copyright Epic Games, Inc. All Rights Reserved.
+
+#include "ChaosClothAsset/OutfitAssetTerminalNode.h"
+#include "ChaosClothAsset/OutfitAsset.h"
+#include "Dataflow/DataflowInputOutput.h"
+
+#include UE_INLINE_GENERATED_CPP_BY_NAME(OutfitAssetTerminalNode)
+
+#define LOCTEXT_NAMESPACE "ChaosOutfitAssetTerminalNode"
+
+FChaosOutfitAssetTerminalNode::FChaosOutfitAssetTerminalNode(const UE::Dataflow::FNodeParameters& InParam, FGuid InGuid)
+	: FDataflowTerminalNode(InParam, InGuid)
+{
+	RegisterInputConnection(&Collection);
+}
+
+void FChaosOutfitAssetTerminalNode::SetAssetValue(TObjectPtr<UObject> Asset, UE::Dataflow::FContext& Context) const
+{
+	if (UChaosOutfitAsset* const OutfitAsset = Cast<UChaosOutfitAsset>(Asset.Get()))
+	{
+		const FManagedArrayCollection& OutfitCollection = GetValue(Context, &Collection);
+		// TODO: Process the Outfit Collection 
+
+		// Rebuild the asset static data
+		OutfitAsset->Build();
+
+		// Asset must be resaved
+		OutfitAsset->MarkPackageDirty();
+	}
+}
+
+#undef LOCTEXT_NAMESPACE
diff --git a/Engine/Plugins/Experimental/ChaosOutfitAsset/Source/ChaosOutfitAssetDataflowNodes/Private/ChaosClothAsset/OutfitAssetTerminalNode.h b/Engine/Plugins/Experimental/ChaosOutfitAsset/Source/ChaosOutfitAssetDataflowNodes/Private/ChaosClothAsset/OutfitAssetTerminalNode.h
new file mode 100644
--- /dev/null
+++ b/Engine/Plugins/Experimental/ChaosOutfitAsset/Source/ChaosOutfitAssetDataflowNodes/Private/ChaosClothAsset/OutfitAssetTerminalNode.h
@@ -0,0 +1,31 @@
+// Copyright Epic Games, Inc. All Rights Reserved.
+
+#pragma once 
+
+#include "Dataflow/DataflowTerminalNode.h"
+#include "GeometryCollection/ManagedArrayCollection.h"
+#include "OutfitAssetTerminalNode.generated.h"
+
+/**
+ * Outfit terminal node to generate an outfit asset from several cloth assets.
+ */
+USTRUCT(Meta = (Experimental, DataflowOutfit, DataflowTerminal))
+struct FChaosOutfitAssetTerminalNode final : public FDataflowTerminalNode
+{
+	GENERATED_USTRUCT_BODY()
+	DATAFLOW_NODE_DEFINE_INTERNAL(FChaosOutfitAssetTerminalNode, "OutfitAssetTerminal", "Outfit", "Outfit Terminal")
+
+public:
+	FChaosOutfitAssetTerminalNode(const UE::Dataflow::FNodeParameters& InParam, FGuid InGuid = FGuid::NewGuid());
+
+private:
+	//~ Begin FDataflowNode interface
+	virtual void SetAssetValue(TObjectPtr<UObject> Asset, UE::Dataflow::FContext& Context) const override;
+	virtual void Evaluate(UE::Dataflow::FContext& Context, const FDataflowOutput* Out) const override {}
+	//~ End FDataflowNode interface
+
+	/** Input outfit collection. */
+	UPROPERTY(Meta = (DataflowInput))
+	FManagedArrayCollection Collection;
+};
+
diff --git a/Engine/Plugins/Experimental/ChaosOutfitAsset/Source/ChaosOutfitAssetEditor/ChaosOutfitAssetEditor.Build.cs b/Engine/Plugins/Experimental/ChaosOutfitAsset/Source/ChaosOutfitAssetEditor/ChaosOutfitAssetEditor.Build.cs
new file mode 100644
--- /dev/null
+++ b/Engine/Plugins/Experimental/ChaosOutfitAsset/Source/ChaosOutfitAssetEditor/ChaosOutfitAssetEditor.Build.cs
@@ -0,0 +1,33 @@
+// Copyright Epic Games, Inc. All Rights Reserved.
+
+using UnrealBuildTool;
+
+public class ChaosOutfitAssetEditor : ModuleRules
+{
+	public ChaosOutfitAssetEditor(ReadOnlyTargetRules Target) : base(Target)
+	{
+		PCHUsage = ModuleRules.PCHUsageMode.UseExplicitOrSharedPCHs;
+
+		PrivateDependencyModuleNames.AddRange(
+			new string[]
+			{
+				"AssetDefinition",
+				"BaseCharacterFXEditor",
+				"ChaosClothAsset",
+				"ChaosClothAssetEngine",
+				"ChaosOutfitAssetEngine",
+				"Core",
+				"CoreUObject",
+				"DataflowEditor",
+				"DataflowEngine",
+				"Engine",
+				"InputCore",
+				"InteractiveToolsFramework",
+                "Projects",
+                "Slate",
+				"SlateCore",
+				"UnrealEd"
+			}
+		);
+	}
+}
diff --git a/Engine/Plugins/Experimental/ChaosOutfitAsset/Source/ChaosOutfitAssetEditor/Private/ChaosClothAsset/AssetDefinition_OutfitAsset.cpp b/Engine/Plugins/Experimental/ChaosOutfitAsset/Source/ChaosOutfitAssetEditor/Private/ChaosClothAsset/AssetDefinition_OutfitAsset.cpp
new file mode 100644
--- /dev/null
+++ b/Engine/Plugins/Experimental/ChaosOutfitAsset/Source/ChaosOutfitAssetEditor/Private/ChaosClothAsset/AssetDefinition_OutfitAsset.cpp
@@ -0,0 +1,60 @@
+// Copyright Epic Games, Inc. All Rights Reserved.
+
+#include "AssetDefinition_OutfitAsset.h"
+#include "ChaosClothAsset/OutfitAsset.h"
+#include "Dataflow/DataflowEditor.h"
+#include "ThumbnailRendering/SceneThumbnailInfo.h"
+
+#define LOCTEXT_NAMESPACE "AssetDefinition_OutfitAsset"
+
+namespace UE::Chaos::ClothAsset::Private
+{
+	static const FLinearColor OutfitAssetColor = FColor(162, 108, 99);
+}
+
+FText UAssetDefinition_OutfitAsset::GetAssetDisplayName() const
+{
+	return NSLOCTEXT("AssetTypeActions", "AssetTypeActions_OutfitAsset", "OutfitAsset");
+}
+
+TSoftClassPtr<UObject> UAssetDefinition_OutfitAsset::GetAssetClass() const
+{
+	return UChaosOutfitAsset::StaticClass();
+}
+
+FLinearColor UAssetDefinition_OutfitAsset::GetAssetColor() const
+{
+	return UE::Chaos::ClothAsset::Private::OutfitAssetColor;
+}
+
+TConstArrayView<FAssetCategoryPath> UAssetDefinition_OutfitAsset::GetAssetCategories() const
+{
+	static const auto Categories = { EAssetCategoryPaths::Physics };
+	return Categories;
+}
+
+UThumbnailInfo* UAssetDefinition_OutfitAsset::LoadThumbnailInfo(const FAssetData& InAsset) const
+{
+	return UE::Editor::FindOrCreateThumbnailInfo(InAsset.GetAsset(), USceneThumbnailInfo::StaticClass());
+}
+
+EAssetCommandResult UAssetDefinition_OutfitAsset::OpenAssets(const FAssetOpenArgs& OpenArgs) const
+{
+	TArray<UChaosOutfitAsset*> OutfitObjects = OpenArgs.LoadObjects<UChaosOutfitAsset>();
+
+	// For now the Dataflow editor only works on one asset at a time
+	if (OutfitObjects.Num() > 0)
+	{
+		UAssetEditorSubsystem* const AssetEditorSubsystem = GEditor->GetEditorSubsystem<UAssetEditorSubsystem>();
+		UDataflowEditor* const AssetEditor = NewObject<UDataflowEditor>(AssetEditorSubsystem, NAME_None, RF_Transient);
+
+		static const TCHAR* const Path = TEXT("/ChaosOutfitAssetEditor/BP_OutfitPreview.BP_OutfitPreview_C");
+		const TSubclassOf<AActor> PreviewClass = StaticLoadClass(AActor::StaticClass(), nullptr, Path);
+		AssetEditor->Initialize({ OutfitObjects[0] }, PreviewClass);
+
+		return EAssetCommandResult::Handled;
+	}
+	return EAssetCommandResult::Unhandled;
+}
+
+#undef LOCTEXT_NAMESPACE
diff --git a/Engine/Plugins/Experimental/ChaosOutfitAsset/Source/ChaosOutfitAssetEditor/Private/ChaosClothAsset/AssetDefinition_OutfitAsset.h b/Engine/Plugins/Experimental/ChaosOutfitAsset/Source/ChaosOutfitAssetEditor/Private/ChaosClothAsset/AssetDefinition_OutfitAsset.h
new file mode 100644
--- /dev/null
+++ b/Engine/Plugins/Experimental/ChaosOutfitAsset/Source/ChaosOutfitAssetEditor/Private/ChaosClothAsset/AssetDefinition_OutfitAsset.h
@@ -0,0 +1,21 @@
+// Copyright Epic Games, Inc. All Rights Reserved.
+
+#pragma once
+
+#include "AssetDefinitionDefault.h"
+#include "AssetDefinition_OutfitAsset.generated.h"
+
+UCLASS(Experimental)
+class UAssetDefinition_OutfitAsset final : public UAssetDefinitionDefault
+{
+	GENERATED_BODY()
+
+private:
+	virtual FText GetAssetDisplayName() const override;
+	virtual TSoftClassPtr<UObject> GetAssetClass() const override;
+	virtual FLinearColor GetAssetColor() const override;
+	virtual TConstArrayView<FAssetCategoryPath> GetAssetCategories() const override;
+	virtual UThumbnailInfo* LoadThumbnailInfo(const FAssetData& InAssetData) const override;
+
+	virtual EAssetCommandResult OpenAssets(const FAssetOpenArgs& OpenArgs) const override;
+};
diff --git a/Engine/Plugins/Experimental/ChaosOutfitAsset/Source/ChaosOutfitAssetEditor/Private/ChaosClothAsset/ChaosOutfitAssetEditorModule.cpp b/Engine/Plugins/Experimental/ChaosOutfitAsset/Source/ChaosOutfitAssetEditor/Private/ChaosClothAsset/ChaosOutfitAssetEditorModule.cpp
new file mode 100644
--- /dev/null
+++ b/Engine/Plugins/Experimental/ChaosOutfitAsset/Source/ChaosOutfitAssetEditor/Private/ChaosClothAsset/ChaosOutfitAssetEditorModule.cpp
@@ -0,0 +1,40 @@
+// Copyright Epic Games, Inc. All Rights Reserved.
+
+#include "ChaosClothAsset/AssetDefinition_OutfitAsset.h"
+#include "ChaosClothAsset/OutfitAssetEditorStyle.h"
+#include "ChaosClothAsset/OutfitEditorCommands.h"
+#include "Modules/ModuleInterface.h"
+#include "Modules/ModuleManager.h"
+
+#define LOCTEXT_NAMESPACE "OutfitAssetEngineModule"
+
+namespace UE::Chaos::ClothAsset
+{
+	class FOutfitAssetEditorModule : public IModuleInterface
+	{
+	public:
+		virtual void StartupModule() override
+		{
+			// Register asset icons
+			FOutfitAssetEditorStyle::Get();
+
+			// Register editor commands
+			FOutfitEditorCommands::Register();
+
+			// TODO Register visualizations (see FChaosClothAssetEditorModule::StartupModule() in ChaosClothAssetEditorModule.cpp)
+		}
+
+		virtual void ShutdownModule() override
+		{
+			if (UObjectInitialized())
+			{
+				// TODO Deregister visualizations (see FChaosClothAssetEditorModule::StartupModule() in ChaosClothAssetEditorModule.cpp)
+				FOutfitEditorCommands::Unregister();
+			}
+		}
+	};
+} // namespace UE::Chaos::ClothAsset
+
+IMPLEMENT_MODULE(UE::Chaos::ClothAsset::FOutfitAssetEditorModule, ChaosOutfitAssetEditor);
+
+#undef LOCTEXT_NAMESPACE
diff --git a/Engine/Plugins/Experimental/ChaosOutfitAsset/Source/ChaosOutfitAssetEditor/Private/ChaosClothAsset/OutfitAssetEditorStyle.cpp b/Engine/Plugins/Experimental/ChaosOutfitAsset/Source/ChaosOutfitAssetEditor/Private/ChaosClothAsset/OutfitAssetEditorStyle.cpp
new file mode 100644
--- /dev/null
+++ b/Engine/Plugins/Experimental/ChaosOutfitAsset/Source/ChaosOutfitAssetEditor/Private/ChaosClothAsset/OutfitAssetEditorStyle.cpp
@@ -0,0 +1,38 @@
+// Copyright Epic Games, Inc. All Rights Reserved.
+
+#include "ChaosClothAsset/OutfitAssetEditorStyle.h"
+#include "Misc/LazySingleton.h"
+#include "Styling/SlateStyleRegistry.h"
+#include "Interfaces/IPluginManager.h"
+
+namespace UE::Chaos::ClothAsset
+{
+	FOutfitAssetEditorStyle::FOutfitAssetEditorStyle()
+		: FSlateStyleSet("OutfitAssetEditorStyle")
+	{
+		const TSharedPtr<IPlugin> ChaosOutfitAssetPlugin = IPluginManager::Get().FindPlugin("ChaosOutfitAsset");
+		if (ChaosOutfitAssetPlugin.IsValid())
+		{
+			SetContentRoot(ChaosOutfitAssetPlugin->GetBaseDir() / TEXT("Resources"));
+
+			Set("ClassIcon.ChaosOutfitAsset", new FSlateVectorImageBrush(RootToContentDir(TEXT("OutfitAsset_16.svg")), FVector2D(16)));
+			Set("ClassThumbnail.ChaosOutfitAsset", new FSlateVectorImageBrush(RootToContentDir(TEXT("OutfitAsset_64.svg")), FVector2D(64)));
+		}
+		FSlateStyleRegistry::RegisterSlateStyle(*this);
+	}
+
+	FOutfitAssetEditorStyle::~FOutfitAssetEditorStyle()
+	{
+		FSlateStyleRegistry::UnRegisterSlateStyle(*this);
+	}
+
+	FOutfitAssetEditorStyle& FOutfitAssetEditorStyle::Get()
+	{
+		return TLazySingleton<FOutfitAssetEditorStyle>::Get();
+	}
+
+	void FOutfitAssetEditorStyle::TearDown()
+	{
+		TLazySingleton<FOutfitAssetEditorStyle>::TearDown();
+	}
+} // namespace UE::Chaos::ClothAsset
diff --git a/Engine/Plugins/Experimental/ChaosOutfitAsset/Source/ChaosOutfitAssetEditor/Private/ChaosClothAsset/OutfitAssetEditorStyle.h b/Engine/Plugins/Experimental/ChaosOutfitAsset/Source/ChaosOutfitAssetEditor/Private/ChaosClothAsset/OutfitAssetEditorStyle.h
new file mode 100644
--- /dev/null
+++ b/Engine/Plugins/Experimental/ChaosOutfitAsset/Source/ChaosOutfitAssetEditor/Private/ChaosClothAsset/OutfitAssetEditorStyle.h
@@ -0,0 +1,24 @@
+// Copyright Epic Games, Inc. All Rights Reserved.
+
+#pragma once
+
+#include "Styling/SlateStyle.h"
+
+namespace UE::Chaos::ClothAsset
+{
+	/**
+	 * Editor style setting up the outfit asset icons in editor.
+	 */
+	class FOutfitAssetEditorStyle final : public FSlateStyleSet
+	{
+	public:
+		static FOutfitAssetEditorStyle& Get();
+		static void TearDown();
+
+	private:
+		friend class FLazySingleton;
+
+		FOutfitAssetEditorStyle();
+		virtual ~FOutfitAssetEditorStyle() override;
+	};
+}  // namespace UE::Chaos::ClothAsset
diff --git a/Engine/Plugins/Experimental/ChaosOutfitAsset/Source/ChaosOutfitAssetEditor/Private/ChaosClothAsset/OutfitAssetFactory.cpp b/Engine/Plugins/Experimental/ChaosOutfitAsset/Source/ChaosOutfitAssetEditor/Private/ChaosClothAsset/OutfitAssetFactory.cpp
new file mode 100644
--- /dev/null
+++ b/Engine/Plugins/Experimental/ChaosOutfitAsset/Source/ChaosOutfitAssetEditor/Private/ChaosClothAsset/OutfitAssetFactory.cpp
@@ -0,0 +1,50 @@
+// Copyright Epic Games, Inc. All Rights Reserved.
+
+#include "ChaosClothAsset/OutfitAssetFactory.h"
+#include "AssetRegistry/AssetRegistryModule.h"
+#include "ChaosClothAsset/OutfitAsset.h"
+#include "Dataflow/DataflowObject.h"
+#include "Misc/PackageName.h"
+
+#include UE_INLINE_GENERATED_CPP_BY_NAME(OutfitAssetFactory)
+
+UChaosOutfitAssetFactory::UChaosOutfitAssetFactory(const FObjectInitializer& ObjectInitializer)
+	: Super(ObjectInitializer)
+{
+	bEditorImport = true;
+	bEditAfterNew = true;
+	SupportedClass = UChaosOutfitAsset::StaticClass();
+}
+
+UObject* UChaosOutfitAssetFactory::FactoryCreateNew(UClass* Class, UObject* Parent, FName Name, EObjectFlags Flags, UObject* /*Context*/, FFeedbackContext* /*Warn*/)
+{
+	UChaosOutfitAsset* const OutfitAsset = NewObject<UChaosOutfitAsset>(Parent, Class, Name, Flags | RF_Transactional | RF_Public | RF_Standalone);
+	OutfitAsset->MarkPackageDirty();
+
+	// Create a new Dataflow asset
+	const FString DataflowPath = FPackageName::GetLongPackagePath(OutfitAsset->GetOutermost()->GetName());
+	const FString OutfitAssetName = OutfitAsset->GetName();
+	const FString DataflowName = FString(TEXT("DF_")) + (OutfitAssetName.StartsWith(TEXT("OA_")) ? OutfitAssetName.RightChop(3) : OutfitAssetName);
+	const FString DataflowPackageName = FPaths::Combine(DataflowPath, DataflowName);
+
+	UPackage* const DataflowPackage = CreatePackage(*DataflowPackageName);
+
+	// Load the cloth template into the new Dataflow asset
+	UDataflow* const Template = LoadObject<UDataflow>(DataflowPackage, TEXT("/ChaosOutfitAsset/OutfitAssetTemplate.OutfitAssetTemplate"));
+	UDataflow* const Dataflow = DuplicateObject(Template, DataflowPackage, FName(DataflowName));
+
+	Dataflow->MarkPackageDirty();
+
+	// Notify the asset registry
+	FAssetRegistryModule::AssetCreated(Dataflow);
+
+	// Set the Dataflow to the cloth asset
+	OutfitAsset->SetDataflow(Dataflow);
+
+	return OutfitAsset;
+}
+
+FString UChaosOutfitAssetFactory::GetDefaultNewAssetName() const
+{
+	return FString(TEXT("OA_NewOutfitAsset"));
+}
diff --git a/Engine/Plugins/Experimental/ChaosOutfitAsset/Source/ChaosOutfitAssetEditor/Private/ChaosClothAsset/OutfitAssetFactory.h b/Engine/Plugins/Experimental/ChaosOutfitAsset/Source/ChaosOutfitAssetEditor/Private/ChaosClothAsset/OutfitAssetFactory.h
new file mode 100644
--- /dev/null
+++ b/Engine/Plugins/Experimental/ChaosOutfitAsset/Source/ChaosOutfitAssetEditor/Private/ChaosClothAsset/OutfitAssetFactory.h
@@ -0,0 +1,26 @@
+// Copyright Epic Games, Inc. All Rights Reserved.
+
+#pragma once
+
+#include "Factories/Factory.h"
+#include "OutfitAssetFactory.generated.h"
+
+/**
+ * Allows the Outfit Asset to be created from the Editor's menus.
+ */
+UCLASS(Experimental)
+class UChaosOutfitAssetFactory final : public UFactory
+{
+	GENERATED_BODY()
+public:
+	UChaosOutfitAssetFactory(const FObjectInitializer& ObjectInitializer);
+
+private:
+	//~ Begin UFactory interface
+	virtual bool CanCreateNew() const override { return true; }
+	virtual bool FactoryCanImport(const FString& Filename) override { return false; }
+	virtual bool ShouldShowInNewMenu() const override { return true; }
+	virtual UObject* FactoryCreateNew(UClass* InClass, UObject* InParent, FName InName, EObjectFlags Flags, UObject* Context, FFeedbackContext* Warn) override;
+	virtual FString GetDefaultNewAssetName() const override;
+	//~ End UFactory interface
+};
diff --git a/Engine/Plugins/Experimental/ChaosOutfitAsset/Source/ChaosOutfitAssetEditor/Private/ChaosClothAsset/OutfitEditorCommands.cpp b/Engine/Plugins/Experimental/ChaosOutfitAsset/Source/ChaosOutfitAssetEditor/Private/ChaosClothAsset/OutfitEditorCommands.cpp
new file mode 100644
--- /dev/null
+++ b/Engine/Plugins/Experimental/ChaosOutfitAsset/Source/ChaosOutfitAssetEditor/Private/ChaosClothAsset/OutfitEditorCommands.cpp
@@ -0,0 +1,42 @@
+// Copyright Epic Games, Inc. All Rights Reserved.
+
+#include "ChaosClothAsset/OutfitEditorCommands.h"
+#include "ChaosClothAsset/OutfitEditorStyle.h"
+
+#define LOCTEXT_NAMESPACE "ChaosClothAssetOutfitEditorCommands"
+
+namespace UE::Chaos::ClothAsset
+{
+	FOutfitEditorCommands::FOutfitEditorCommands()
+		: TBaseCharacterFXEditorCommands<FOutfitEditorCommands>(
+			"ChaosClothOutfitEditor",
+			LOCTEXT("ContextDescription", "Outfit Editor"), 
+			NAME_None, // Parent
+			FOutfitEditorStyle::Get().GetStyleSetName())
+	{
+	}
+
+	void FOutfitEditorCommands::RegisterCommands()
+	{
+		TBaseCharacterFXEditorCommands::RegisterCommands();
+
+		UI_COMMAND(OpenOutfitEditor, "Outfit Editor", "Open the Outfit Editor window", EUserInterfaceActionType::Button, FInputChord());
+	}
+
+	void FOutfitEditorCommands::UpdateToolCommandBinding(UInteractiveTool* Tool, TSharedPtr<FUICommandList> UICommandList, bool bUnbind)
+	{
+		if (FOutfitEditorCommands::IsRegistered())
+		{
+			if (bUnbind)
+			{
+				FOutfitEditorCommands::Get().UnbindActiveCommands(UICommandList);
+			}
+			else
+			{
+				FOutfitEditorCommands::Get().BindCommandsForCurrentTool(UICommandList, Tool);
+			}
+		}
+	}
+} // namespace UE::Chaos::ClothAsset
+
+#undef LOCTEXT_NAMESPACE
diff --git a/Engine/Plugins/Experimental/ChaosOutfitAsset/Source/ChaosOutfitAssetEditor/Private/ChaosClothAsset/OutfitEditorCommands.h b/Engine/Plugins/Experimental/ChaosOutfitAsset/Source/ChaosOutfitAssetEditor/Private/ChaosClothAsset/OutfitEditorCommands.h
new file mode 100644
--- /dev/null
+++ b/Engine/Plugins/Experimental/ChaosOutfitAsset/Source/ChaosOutfitAssetEditor/Private/ChaosClothAsset/OutfitEditorCommands.h
@@ -0,0 +1,32 @@
+// Copyright Epic Games, Inc. All Rights Reserved.
+
+#pragma once
+
+#include "BaseCharacterFXEditorCommands.h"
+
+namespace UE::Chaos::ClothAsset
+{
+	class FOutfitEditorCommands final : public TBaseCharacterFXEditorCommands<FOutfitEditorCommands>
+	{
+	public:
+		FOutfitEditorCommands();
+
+		/**
+		 * Add or remove commands relevant to Tool to the given UICommandList.
+		 * Call this when the active tool changes (e.g. on ToolManager.OnToolStarted / OnToolEnded).
+		 * @param bUnbind if true, commands are removed, otherwise added.
+		 */
+		static void UpdateToolCommandBinding(UInteractiveTool* Tool, TSharedPtr<FUICommandList> UICommandList, bool bUnbind = false);
+
+	private:
+		friend TCommands<FOutfitEditorCommands>;
+
+		//~ Begin TBaseCharacterFXEditorCommands<FOutfitEditorCommands> interface
+		virtual void RegisterCommands() override;
+
+		virtual void GetToolDefaultObjectList(TArray<UInteractiveTool*>& ToolCDOs) override {}
+		//~ End TBaseCharacterFXEditorCommands<FOutfitEditorCommands> interface
+
+		TSharedPtr<FUICommandInfo> OpenOutfitEditor;
+	};
+} // namespace UE::Chaos::ClothAsset
diff --git a/Engine/Plugins/Experimental/ChaosOutfitAsset/Source/ChaosOutfitAssetEditor/Private/ChaosClothAsset/OutfitEditorStyle.cpp b/Engine/Plugins/Experimental/ChaosOutfitAsset/Source/ChaosOutfitAssetEditor/Private/ChaosClothAsset/OutfitEditorStyle.cpp
new file mode 100644
--- /dev/null
+++ b/Engine/Plugins/Experimental/ChaosOutfitAsset/Source/ChaosOutfitAssetEditor/Private/ChaosClothAsset/OutfitEditorStyle.cpp
@@ -0,0 +1,29 @@
+// Copyright Epic Games, Inc. All Rights Reserved.
+
+#include "ChaosClothAsset/OutfitEditorStyle.h"
+#include "Misc/LazySingleton.h"
+#include "Styling/SlateStyleRegistry.h"
+
+namespace UE::Chaos::ClothAsset
+{
+	FOutfitEditorStyle::FOutfitEditorStyle()
+		: FSlateStyleSet(StyleName)
+	{
+		FSlateStyleRegistry::RegisterSlateStyle(*this);
+	}
+
+	FOutfitEditorStyle::~FOutfitEditorStyle()
+	{
+		FSlateStyleRegistry::UnRegisterSlateStyle(*this);
+	}
+
+	FOutfitEditorStyle& FOutfitEditorStyle::Get()
+	{
+		return TLazySingleton<FOutfitEditorStyle>::Get();
+	}
+
+	void FOutfitEditorStyle::TearDown()
+	{
+		TLazySingleton<FOutfitEditorStyle>::TearDown();
+	}
+} // namespace UE::Chaos::ClothAsset
diff --git a/Engine/Plugins/Experimental/ChaosOutfitAsset/Source/ChaosOutfitAssetEditor/Private/ChaosClothAsset/OutfitEditorStyle.h b/Engine/Plugins/Experimental/ChaosOutfitAsset/Source/ChaosOutfitAssetEditor/Private/ChaosClothAsset/OutfitEditorStyle.h
new file mode 100644
--- /dev/null
+++ b/Engine/Plugins/Experimental/ChaosOutfitAsset/Source/ChaosOutfitAssetEditor/Private/ChaosClothAsset/OutfitEditorStyle.h
@@ -0,0 +1,26 @@
+// Copyright Epic Games, Inc. All Rights Reserved.
+
+#pragma once
+
+#include "Styling/SlateStyle.h"
+
+namespace UE::Chaos::ClothAsset
+{
+	/**
+	 * Slate style set for Cloth Editor.
+	 */
+	class CHAOSOUTFITASSETEDITOR_API FOutfitEditorStyle final : public FSlateStyleSet
+	{
+	public:
+		static inline const FName StyleName = "OutfitEditorStyle";
+
+		static FOutfitEditorStyle& Get();
+		static void TearDown();
+
+	private:
+		friend class FLazySingleton;
+
+		FOutfitEditorStyle();
+		virtual ~FOutfitEditorStyle() override;
+	};
+} // namespace UE::Chaos::ClothAsset
diff --git a/Engine/Plugins/Experimental/ChaosOutfitAsset/Source/ChaosOutfitAssetEngine/ChaosOutfitAssetEngine.Build.cs b/Engine/Plugins/Experimental/ChaosOutfitAsset/Source/ChaosOutfitAssetEngine/ChaosOutfitAssetEngine.Build.cs
new file mode 100644
--- /dev/null
+++ b/Engine/Plugins/Experimental/ChaosOutfitAsset/Source/ChaosOutfitAssetEngine/ChaosOutfitAssetEngine.Build.cs
@@ -0,0 +1,22 @@
+// Copyright Epic Games, Inc. All Rights Reserved.
+
+using UnrealBuildTool;
+
+public class ChaosOutfitAssetEngine : ModuleRules
+{
+	public ChaosOutfitAssetEngine(ReadOnlyTargetRules Target) : base(Target)
+	{
+		PCHUsage = ModuleRules.PCHUsageMode.UseExplicitOrSharedPCHs;
+
+		PrivateDependencyModuleNames.AddRange(
+			new string[]
+			{
+				"Core",
+				"CoreUObject",
+				"ChaosClothAsset",
+				"ChaosClothAssetEngine",
+				"Engine"
+			}
+		);
+	}
+}
diff --git a/Engine/Plugins/Experimental/ChaosOutfitAsset/Source/ChaosOutfitAssetEngine/Private/ChaosClothAsset/ChaosOutfitAssetEngineModule.cpp b/Engine/Plugins/Experimental/ChaosOutfitAsset/Source/ChaosOutfitAssetEngine/Private/ChaosClothAsset/ChaosOutfitAssetEngineModule.cpp
new file mode 100644
--- /dev/null
+++ b/Engine/Plugins/Experimental/ChaosOutfitAsset/Source/ChaosOutfitAssetEngine/Private/ChaosClothAsset/ChaosOutfitAssetEngineModule.cpp
@@ -0,0 +1,25 @@
+// Copyright Epic Games, Inc. All Rights Reserved.
+
+#include "Modules/ModuleInterface.h"
+#include "Modules/ModuleManager.h"
+
+#define LOCTEXT_NAMESPACE "OutfitAssetEngineModule"
+
+namespace UE::Chaos::ClothAsset
+{
+	class FOutfitAssetEngineModule : public IModuleInterface
+	{
+	public:
+		virtual void StartupModule() override
+		{
+		}
+
+		virtual void ShutdownModule() override
+		{
+		}
+	};
+}  // namespace UE::Chaos::ClothAsset
+
+IMPLEMENT_MODULE(UE::Chaos::ClothAsset::FOutfitAssetEngineModule, ChaosOutfitAssetEngine);
+
+#undef LOCTEXT_NAMESPACE
diff --git a/Engine/Plugins/Experimental/ChaosOutfitAsset/Source/ChaosOutfitAssetEngine/Private/ChaosClothAsset/OutfitAsset.cpp b/Engine/Plugins/Experimental/ChaosOutfitAsset/Source/ChaosOutfitAssetEngine/Private/ChaosClothAsset/OutfitAsset.cpp
new file mode 100644
--- /dev/null
+++ b/Engine/Plugins/Experimental/ChaosOutfitAsset/Source/ChaosOutfitAssetEngine/Private/ChaosClothAsset/OutfitAsset.cpp
@@ -0,0 +1,19 @@
+// Copyright Epic Games, Inc. All Rights Reserved.
+
+#include "ChaosClothAsset/OutfitAsset.h"
+
+#include UE_INLINE_GENERATED_CPP_BY_NAME(OutfitAsset)
+
+UChaosOutfitAsset::UChaosOutfitAsset(const FObjectInitializer& ObjectInitializer)
+	: Super(ObjectInitializer)
+{
+	SetDataflowTerminal(TEXT("OutfitAssetTerminal"));
+}
+
+UChaosOutfitAsset::UChaosOutfitAsset(FVTableHelper& Helper)
+	: Super(Helper)
+{
+}
+
+UChaosOutfitAsset::~UChaosOutfitAsset() = default;
+
diff --git a/Engine/Plugins/Experimental/ChaosOutfitAsset/Source/ChaosOutfitAssetEngine/Public/ChaosClothAsset/OutfitAsset.h b/Engine/Plugins/Experimental/ChaosOutfitAsset/Source/ChaosOutfitAssetEngine/Public/ChaosClothAsset/OutfitAsset.h
new file mode 100644
--- /dev/null
+++ b/Engine/Plugins/Experimental/ChaosOutfitAsset/Source/ChaosOutfitAssetEngine/Public/ChaosClothAsset/OutfitAsset.h
@@ -0,0 +1,22 @@
+// Copyright Epic Games, Inc. All Rights Reserved.
+
+#pragma once
+
+#include "ChaosClothAsset/ClothAsset.h"
+#include "OutfitAsset.generated.h"
+
+class USkeletalMesh;
+class UDataflow;
+
+/**
+ * Outfit asset for character clothing and simulation.
+ */
+UCLASS(Experimental, HideCategories = Object, BlueprintType)
+class CHAOSOUTFITASSETENGINE_API UChaosOutfitAsset final : public UChaosClothAsset
+{
+	GENERATED_BODY()
+public:
+	UChaosOutfitAsset(const FObjectInitializer& ObjectInitializer);
+	UChaosOutfitAsset(FVTableHelper& Helper);  // This is declared so we can use TUniquePtr<FClothSimulationModel> with just a forward declare of that class
+	virtual ~UChaosOutfitAsset() override;
+};
