CIS - Fixed issue 746667: Compile errors in Module.ChaosOutfitAssetEditor.cpp.

#rnx

diff --git a/Engine/Plugins/Experimental/ChaosOutfitAsset/Source/ChaosOutfitAssetEditor/Private/ChaosClothAsset/OutfitAssetEditorStyle.h b/Engine/Plugins/Experimental/ChaosOutfitAsset/Source/ChaosOutfitAssetEditor/Private/ChaosClothAsset/OutfitAssetEditorStyle.h
--- a/Engine/Plugins/Experimental/ChaosOutfitAsset/Source/ChaosOutfitAssetEditor/Private/ChaosClothAsset/OutfitAssetEditorStyle.h
+++ b/Engine/Plugins/Experimental/ChaosOutfitAsset/Source/ChaosOutfitAssetEditor/Private/ChaosClothAsset/OutfitAssetEditorStyle.h
@@ -15,9 +15,6 @@ namespace UE::Chaos::ClothAsset
 		static FOutfitAssetEditorStyle& Get();
 		static void TearDown();
 
-	private:
-		friend class FLazySingleton;
-
 		FOutfitAssetEditorStyle();
 		virtual ~FOutfitAssetEditorStyle() override;
 	};
diff --git a/Engine/Plugins/Experimental/ChaosOutfitAsset/Source/ChaosOutfitAssetEditor/Private/ChaosClothAsset/OutfitEditorStyle.h b/Engine/Plugins/Experimental/ChaosOutfitAsset/Source/ChaosOutfitAssetEditor/Private/ChaosClothAsset/OutfitEditorStyle.h
--- a/Engine/Plugins/Experimental/ChaosOutfitAsset/Source/ChaosOutfitAssetEditor/Private/ChaosClothAsset/OutfitEditorStyle.h
+++ b/Engine/Plugins/Experimental/ChaosOutfitAsset/Source/ChaosOutfitAssetEditor/Private/ChaosClothAsset/OutfitEditorStyle.h
@@ -17,9 +17,6 @@ namespace UE::Chaos::ClothAsset
 		static FOutfitEditorStyle& Get();
 		static void TearDown();
 
-	private:
-		friend class FLazySingleton;
-
 		FOutfitEditorStyle();
 		virtual ~FOutfitEditorStyle() override;
 	};
