Fix for CIS error when adding variables to ClothAsset by adding DataflowEngine  dependency to ChaosOutfitAssetEngine. This shouldn't be necessary because ChaosClothAsset has it as a public dependency; seems to be something wrong with the build system, but this works around it.

diff --git a/Engine/Plugins/Experimental/ChaosOutfitAsset/Source/ChaosOutfitAssetEngine/ChaosOutfitAssetEngine.Build.cs b/Engine/Plugins/Experimental/ChaosOutfitAsset/Source/ChaosOutfitAssetEngine/ChaosOutfitAssetEngine.Build.cs
--- a/Engine/Plugins/Experimental/ChaosOutfitAsset/Source/ChaosOutfitAssetEngine/ChaosOutfitAssetEngine.Build.cs
+++ b/Engine/Plugins/Experimental/ChaosOutfitAsset/Source/ChaosOutfitAssetEngine/ChaosOutfitAssetEngine.Build.cs
@@ -8,6 +8,14 @@ public class ChaosOutfitAssetEngine : ModuleRules
 	{
 		PCHUsage = ModuleRules.PCHUsageMode.UseExplicitOrSharedPCHs;
 
+		PublicDependencyModuleNames.AddRange(
+			new string[]
+			{
+				"DataflowCore",
+				"DataflowEngine",
+			}
+		);
+
 		PrivateDependencyModuleNames.AddRange(
 			new string[]
 			{
