Chaos Cloth Outfit - Fixed the path of the Outfit Asset preview blueprint.

#rnx

diff --git a/Engine/Plugins/Experimental/ChaosOutfitAsset/Source/ChaosOutfitAssetEditor/Private/ChaosClothAsset/AssetDefinition_OutfitAsset.cpp b/Engine/Plugins/Experimental/ChaosOutfitAsset/Source/ChaosOutfitAssetEditor/Private/ChaosClothAsset/AssetDefinition_OutfitAsset.cpp
--- a/Engine/Plugins/Experimental/ChaosOutfitAsset/Source/ChaosOutfitAssetEditor/Private/ChaosClothAsset/AssetDefinition_OutfitAsset.cpp
+++ b/Engine/Plugins/Experimental/ChaosOutfitAsset/Source/ChaosOutfitAssetEditor/Private/ChaosClothAsset/AssetDefinition_OutfitAsset.cpp
@@ -48,7 +48,7 @@ EAssetCommandResult UAssetDefinition_OutfitAsset::OpenAssets(const FAssetOpenArg
 		UAssetEditorSubsystem* const AssetEditorSubsystem = GEditor->GetEditorSubsystem<UAssetEditorSubsystem>();
 		UDataflowEditor* const AssetEditor = NewObject<UDataflowEditor>(AssetEditorSubsystem, NAME_None, RF_Transient);
 
-		static const TCHAR* const Path = TEXT("/ChaosOutfitAssetEditor/BP_OutfitPreview.BP_OutfitPreview_C");
+		static const TCHAR* const Path = TEXT("/ChaosOutfitAsset/BP_OutfitPreview.BP_OutfitPreview_C");
 		const TSubclassOf<AActor> PreviewClass = StaticLoadClass(AActor::StaticClass(), nullptr, Path);
 		AssetEditor->Initialize({ OutfitObjects[0] }, PreviewClass);
 
