Chaos Outfit Asset - Added MakeOutfitNode and MergeOutfits nodes, and removed the ClothAssetToOutfitCollection stub node.

#rnx

diff --git a/Engine/Plugins/Experimental/ChaosOutfitAsset/Source/ChaosOutfitAssetDataflowNodes/Private/ChaosClothAsset/ChaosOutfitAssetDataflowNodesModule.cpp b/Engine/Plugins/Experimental/ChaosOutfitAsset/Source/ChaosOutfitAssetDataflowNodes/Private/ChaosClothAsset/ChaosOutfitAssetDataflowNodesModule.cpp
--- a/Engine/Plugins/Experimental/ChaosOutfitAsset/Source/ChaosOutfitAssetDataflowNodes/Private/ChaosClothAsset/ChaosOutfitAssetDataflowNodesModule.cpp
+++ b/Engine/Plugins/Experimental/ChaosOutfitAsset/Source/ChaosOutfitAssetDataflowNodes/Private/ChaosClothAsset/ChaosOutfitAssetDataflowNodesModule.cpp
@@ -1,8 +1,13 @@
 // Copyright Epic Games, Inc. All Rights Reserved.
 
-#include "ChaosClothAsset/ClothAssetToOutfitCollectionNode.h"
+#include "ChaosClothAsset/ClothAssetAnyType.h"
 #include "ChaosClothAsset/GetClothAssetNode.h"
+#include "ChaosClothAsset/GetOutfitAssetNode.h"
+#include "ChaosClothAsset/MakeOutfitNode.h"
+#include "ChaosClothAsset/MergeOutfitsNode.h"
 #include "ChaosClothAsset/OutfitAssetTerminalNode.h"
+#include "Dataflow/DataflowAnyTypeRegistry.h"
+#include "Dataflow/DataflowNode.h"
 #include "Dataflow/DataflowNodeColorsRegistry.h"
 #include "Dataflow/DataflowNodeFactory.h"
 #include "Modules/ModuleInterface.h"
@@ -20,13 +25,27 @@ namespace UE::Chaos::ClothAsset
 		static void RegisterDataflowNodes()
 		{
 			DATAFLOW_NODE_REGISTER_CREATION_FACTORY_NODE_COLORS_BY_CATEGORY("Outfit", OutfitAssetNodeHeaderColor, OutfitAssetNodeBodyColor);
-			DATAFLOW_NODE_REGISTER_CREATION_FACTORY(FChaosClothAssetToOutfitCollectionNode);
 			DATAFLOW_NODE_REGISTER_CREATION_FACTORY(FChaosGetClothAssetNode);
+			DATAFLOW_NODE_REGISTER_CREATION_FACTORY(FChaosGetOutfitAssetNode);
+			DATAFLOW_NODE_REGISTER_CREATION_FACTORY(FChaosOutfitAssetMakeOutfitNode);
+			DATAFLOW_NODE_REGISTER_CREATION_FACTORY(FChaosOutfitAssetMergeOutfitsNode);
 			DATAFLOW_NODE_REGISTER_CREATION_FACTORY(FChaosOutfitAssetTerminalNode);
 PRAGMA_DISABLE_DEPRECATION_WARNINGS
 			// Move any deprecated node registrations under here
 PRAGMA_ENABLE_DEPRECATION_WARNINGS
 		}
+
+		static void RegisterAnyTypes()
+		{
+			constexpr float DefaultWireThickness = 1.5f;
+			// Register Anytype for UChaosClothAssetBase (UChaosClothAsset and UChaosOutfitAsset)
+			UE_DATAFLOW_REGISTER_ANYTYPE(FChaosClothAssetAnyType);
+			DATAFLOW_NODE_REGISTER_CREATION_FACTORY_PIN_SETTINGS_BY_TYPE("FChaosClothAssetAnyType", OutfitAssetNodeHeaderColor, DefaultWireThickness);
+			UE_DATAFLOW_REGISTER_ANYTYPE(FChaosClothAssetArrayAnyType);
+			DATAFLOW_NODE_REGISTER_CREATION_FACTORY_PIN_SETTINGS_BY_TYPE("FChaosClothAssetArrayAnyType", OutfitAssetNodeHeaderColor, DefaultWireThickness);
+			UE_DATAFLOW_REGISTER_ANYTYPE(FChaosClothAssetOrArrayAnyType);
+			DATAFLOW_NODE_REGISTER_CREATION_FACTORY_PIN_SETTINGS_BY_TYPE("FChaosClothAssetOrArrayAnyType", OutfitAssetNodeHeaderColor, DefaultWireThickness);
+		}
 	}  // namespace Private
 
 	class FOutfitAssetDataflowNodesModule : public IModuleInterface
@@ -34,6 +53,7 @@ PRAGMA_ENABLE_DEPRECATION_WARNINGS
 	public:
 		virtual void StartupModule() override
 		{
+			Private::RegisterAnyTypes();
 			Private::RegisterDataflowNodes();
 		}
 
diff --git a/Engine/Plugins/Experimental/ChaosOutfitAsset/Source/ChaosOutfitAssetDataflowNodes/Private/ChaosClothAsset/ClothAssetAnyType.h b/Engine/Plugins/Experimental/ChaosOutfitAsset/Source/ChaosOutfitAssetDataflowNodes/Private/ChaosClothAsset/ClothAssetAnyType.h
new file mode 100644
--- /dev/null
+++ b/Engine/Plugins/Experimental/ChaosOutfitAsset/Source/ChaosOutfitAssetDataflowNodes/Private/ChaosClothAsset/ClothAssetAnyType.h
@@ -0,0 +1,130 @@
+// Copyright Epic Games, Inc. All Rights Reserved.
+
+#pragma once 
+
+#include "Dataflow/DataflowAnyType.h"
+#include "Dataflow/DataflowTypePolicy.h"
+#include "ChaosClothAsset/ClothAsset.h"
+#include "ChaosClothAsset/OutfitAsset.h"
+#include "Misc/TVariant.h"
+
+#include "ClothAssetAnyType.generated.h"
+
+USTRUCT()
+struct FChaosClothAssetOrArrayType
+{
+	GENERATED_BODY()
+
+	TVariant<TObjectPtr<UChaosClothAssetBase>, TArray<TObjectPtr<UChaosClothAssetBase>>> AssetOrArray;
+
+	FChaosClothAssetOrArrayType() = default;
+
+	template <typename T UE_REQUIRES(TIsDerivedFrom<T, UChaosClothAssetBase>::Value)>
+	FChaosClothAssetOrArrayType(const TObjectPtr<T>& Asset)
+		: AssetOrArray(TInPlaceType<TObjectPtr<UChaosClothAssetBase>>(), Asset)
+	{}
+
+	template <typename T UE_REQUIRES(TIsDerivedFrom<T, UChaosClothAssetBase>::Value)>
+	FChaosClothAssetOrArrayType(const TArray<TObjectPtr<T>>& Array)
+		: AssetOrArray(TInPlaceType<TArray<TObjectPtr<UChaosClothAssetBase>>>(), Array)
+	{}
+
+	bool IsArray() const
+	{
+		return AssetOrArray.IsType<TArray<TObjectPtr<UChaosClothAssetBase>>>();
+	}
+
+	const TObjectPtr<UChaosClothAssetBase>& Get() const
+	{
+		check(!IsArray());
+		return AssetOrArray.Get<TObjectPtr<UChaosClothAssetBase>>();
+	}
+
+	const TArray<TObjectPtr<UChaosClothAssetBase>>& GetArray() const
+	{
+		check(IsArray());
+		return AssetOrArray.Get<TArray<TObjectPtr<UChaosClothAssetBase>>>();
+	}
+
+	bool Serialize(FArchive& Ar)
+	{
+		Ar << AssetOrArray;
+		return true;
+	}
+};
+template<>
+struct TStructOpsTypeTraits<FChaosClothAssetOrArrayType> : public TStructOpsTypeTraitsBase2<FChaosClothAssetOrArrayType>
+{
+	enum
+	{
+		WithSerializer = true,
+	};
+};
+
+UE_DATAFLOW_POLICY_DECLARE_TYPENAME(TObjectPtr<UChaosClothAssetBase>)
+UE_DATAFLOW_POLICY_DECLARE_TYPENAME(TObjectPtr<UChaosClothAsset>)
+UE_DATAFLOW_POLICY_DECLARE_TYPENAME(TObjectPtr<UChaosOutfitAsset>)
+UE_DATAFLOW_POLICY_DECLARE_TYPENAME(FChaosClothAssetOrArrayType)
+
+/** UChaosClothAssetBase types. */
+USTRUCT()
+struct FChaosClothAssetAnyType : public FDataflowAnyType
+{
+	GENERATED_BODY()
+
+	struct FTypePolicy :
+		public TDataflowMultiTypePolicy<
+		TObjectPtr<UChaosClothAssetBase>,
+		TObjectPtr<UChaosClothAsset>,
+		TObjectPtr<UChaosOutfitAsset>>
+	{};
+
+	using FPolicyType = FTypePolicy;
+	using FStorageType = TObjectPtr<UChaosClothAssetBase>;
+
+	UPROPERTY(EditAnywhere, Category = Value)
+	TObjectPtr<UChaosClothAssetBase> Value = nullptr;
+};
+
+/** UChaosClothAssetBase array types. */
+USTRUCT()
+struct FChaosClothAssetArrayAnyType : public FDataflowAnyType
+{
+	GENERATED_BODY()
+
+	struct FArrayPolicy :
+		public TDataflowMultiTypePolicy<
+			TArray<TObjectPtr<UChaosClothAssetBase>>,
+			TArray<TObjectPtr<UChaosClothAsset>>,
+			TArray<TObjectPtr<UChaosOutfitAsset>>>
+	{};
+
+	using FPolicyType = FArrayPolicy;
+	using FStorageType = TArray<TObjectPtr<UChaosClothAssetBase>>;
+
+	UPROPERTY(EditAnywhere, Category = Value)
+	TArray<TObjectPtr<UChaosClothAssetBase>> Value;
+};
+
+/** UChaosClothAssetBase or array types. */
+USTRUCT()
+struct FChaosClothAssetOrArrayAnyType : public FDataflowAnyType
+{
+	GENERATED_BODY()
+
+	struct FAssetOrArrayPolicy :
+		public TDataflowMultiTypePolicy<
+			TObjectPtr<UChaosClothAssetBase>,
+			TObjectPtr<UChaosClothAsset>,
+			TObjectPtr<UChaosOutfitAsset>,
+			TArray<TObjectPtr<UChaosClothAssetBase>>,
+			TArray<TObjectPtr<UChaosClothAsset>>,
+			TArray<TObjectPtr<UChaosOutfitAsset>>>
+	{};
+
+	using FPolicyType = FAssetOrArrayPolicy;
+	using FStorageType = FChaosClothAssetOrArrayType;
+
+	UPROPERTY(EditAnywhere, Category = Value)
+	FChaosClothAssetOrArrayType Value;
+};
diff --git a/Engine/Plugins/Experimental/ChaosOutfitAsset/Source/ChaosOutfitAssetDataflowNodes/Private/ChaosClothAsset/ClothAssetToOutfitCollectionNode.cpp b/Engine/Plugins/Experimental/ChaosOutfitAsset/Source/ChaosOutfitAssetDataflowNodes/Private/ChaosClothAsset/ClothAssetToOutfitCollectionNode.cpp
deleted file mode 100644
diff --git a/Engine/Plugins/Experimental/ChaosOutfitAsset/Source/ChaosOutfitAssetDataflowNodes/Private/ChaosClothAsset/ClothAssetToOutfitCollectionNode.h b/Engine/Plugins/Experimental/ChaosOutfitAsset/Source/ChaosOutfitAssetDataflowNodes/Private/ChaosClothAsset/ClothAssetToOutfitCollectionNode.h
deleted file mode 100644
diff --git a/Engine/Plugins/Experimental/ChaosOutfitAsset/Source/ChaosOutfitAssetDataflowNodes/Private/ChaosClothAsset/GetOutfitAssetNode.cpp b/Engine/Plugins/Experimental/ChaosOutfitAsset/Source/ChaosOutfitAssetDataflowNodes/Private/ChaosClothAsset/GetOutfitAssetNode.cpp
new file mode 100644
--- /dev/null
+++ b/Engine/Plugins/Experimental/ChaosOutfitAsset/Source/ChaosOutfitAssetDataflowNodes/Private/ChaosClothAsset/GetOutfitAssetNode.cpp
@@ -0,0 +1,24 @@
+// Copyright Epic Games, Inc. All Rights Reserved.
+
+#include "ChaosClothAsset/GetOutfitAssetNode.h"
+#include "ChaosClothAsset/OutfitAsset.h"
+
+#include UE_INLINE_GENERATED_CPP_BY_NAME(GetOutfitAssetNode)
+
+#define LOCTEXT_NAMESPACE "ChaosGetOutfitAssetNode"
+
+FChaosGetOutfitAssetNode::FChaosGetOutfitAssetNode(const UE::Dataflow::FNodeParameters& InParam, FGuid InGuid)
+	: FDataflowNode(InParam, InGuid)
+{
+	RegisterOutputConnection(&OutfitAsset);
+}
+
+void FChaosGetOutfitAssetNode::Evaluate(UE::Dataflow::FContext& Context, const FDataflowOutput* Out) const
+{
+	if (Out->IsA<TObjectPtr<const UChaosOutfitAsset>>(&OutfitAsset))
+	{
+		SetValue(Context, OutfitAsset, &OutfitAsset);
+	}
+}
+
+#undef LOCTEXT_NAMESPACE
diff --git a/Engine/Plugins/Experimental/ChaosOutfitAsset/Source/ChaosOutfitAssetDataflowNodes/Private/ChaosClothAsset/GetOutfitAssetNode.h b/Engine/Plugins/Experimental/ChaosOutfitAsset/Source/ChaosOutfitAssetDataflowNodes/Private/ChaosClothAsset/GetOutfitAssetNode.h
new file mode 100644
--- /dev/null
+++ b/Engine/Plugins/Experimental/ChaosOutfitAsset/Source/ChaosOutfitAssetDataflowNodes/Private/ChaosClothAsset/GetOutfitAssetNode.h
@@ -0,0 +1,31 @@
+// Copyright Epic Games, Inc. All Rights Reserved.
+
+#pragma once 
+
+#include "Dataflow/DataflowNode.h"
+#include "GeometryCollection/ManagedArrayCollection.h"
+#include "GetOutfitAssetNode.generated.h"
+
+class UChaosOutfitAsset;
+
+/**
+ * Get an outfit asset object into the graph.
+ */
+USTRUCT(Meta = (Experimental, DataflowOutfit))
+struct FChaosGetOutfitAssetNode final : public FDataflowNode
+{
+	GENERATED_USTRUCT_BODY()
+	DATAFLOW_NODE_DEFINE_INTERNAL(FChaosGetOutfitAssetNode, "GetOutfitAsset", "Outfit", "Get Outfit Asset")
+
+public:
+	FChaosGetOutfitAssetNode(const UE::Dataflow::FNodeParameters& InParam, FGuid InGuid = FGuid::NewGuid());
+
+private:
+	//~ Begin FDataflowNode interface
+	virtual void Evaluate(UE::Dataflow::FContext& Context, const FDataflowOutput* Out) const override;
+	//~ End FDataflowNode interface
+
+	/** The Outfit Asset to import into a collection. */
+	UPROPERTY(EditAnywhere, Category = "GetOutfitAsset", Meta = (DataflowOutput))
+	TObjectPtr<const UChaosOutfitAsset> OutfitAsset;
+};
diff --git a/Engine/Plugins/Experimental/ChaosOutfitAsset/Source/ChaosOutfitAssetDataflowNodes/Private/ChaosClothAsset/MakeOutfitNode.cpp b/Engine/Plugins/Experimental/ChaosOutfitAsset/Source/ChaosOutfitAssetDataflowNodes/Private/ChaosClothAsset/MakeOutfitNode.cpp
new file mode 100644
--- /dev/null
+++ b/Engine/Plugins/Experimental/ChaosOutfitAsset/Source/ChaosOutfitAssetDataflowNodes/Private/ChaosClothAsset/MakeOutfitNode.cpp
@@ -0,0 +1,135 @@
+// Copyright Epic Games, Inc. All Rights Reserved.
+
+#include "ChaosClothAsset/MakeOutfitNode.h"
+#include "ChaosClothAsset/ClothAssetAnyType.h"
+#include "ChaosClothAsset/Outfit.h"
+
+#include UE_INLINE_GENERATED_CPP_BY_NAME(MakeOutfitNode)
+
+#define LOCTEXT_NAMESPACE "ChaosOutfitAssetMakeOutfitNode"
+
+FChaosOutfitAssetMakeOutfitNode::FChaosOutfitAssetMakeOutfitNode(const UE::Dataflow::FNodeParameters& InParam, FGuid InGuid)
+	: FDataflowNode(InParam, InGuid)
+{
+	RegisterOutputConnection(&Outfit);
+
+	for (int32 Index = 0; Index < NumInitialClothAssets; ++Index)
+	{
+		AddPins();
+	}
+	check(GetNumInputs() == NumRequiredInputs + NumInitialClothAssets);  // Update NumRequiredInputs when adding inputs (used by Serialize)
+}
+
+void FChaosOutfitAssetMakeOutfitNode::Evaluate(UE::Dataflow::FContext& Context, const FDataflowOutput* Out) const
+{
+	if (Out->IsA(&Outfit))
+	{
+		TObjectPtr<UChaosOutfit> OutOutfit = NewObject<UChaosOutfit>();
+
+		auto AddAsset = [&OutOutfit](const TObjectPtr<UChaosClothAssetBase>& InAsset)
+			{
+				if (InAsset && InAsset->GetNumClothSimulationModels())  // TODO Use future number of piece API instead
+				{
+					OutOutfit->Add(*InAsset);
+				}
+			};
+
+		for (int32 InputIndex = 0; InputIndex < ClothAssets.Num(); ++InputIndex)
+		{
+			const FChaosClothAssetOrArrayType& InAssetOrArray = GetValue(Context, &ClothAssets[InputIndex]);
+			if (!InAssetOrArray.IsArray())
+			{
+				AddAsset(InAssetOrArray.Get());
+			}
+			else
+			{
+				const TArray<TObjectPtr<UChaosClothAssetBase>>& InArray = InAssetOrArray.GetArray();
+
+				for (int32 ElementIndex = 0; ElementIndex < InArray.Num(); ++ElementIndex)
+				{
+					AddAsset(InArray[ElementIndex]);
+				}
+			}
+		}
+
+		SetValue<TObjectPtr<const UChaosOutfit>>(Context, OutOutfit, &Outfit);
+	}
+}
+
+TArray<UE::Dataflow::FPin> FChaosOutfitAssetMakeOutfitNode::AddPins()
+{
+	const int32 Index = ClothAssets.AddDefaulted();
+	const FDataflowInput& Input = RegisterInputArrayConnection(GetConnectionReference(Index));
+	return { { UE::Dataflow::FPin::EDirection::INPUT, Input.GetType(), Input.GetName() } };
+}
+
+TArray<UE::Dataflow::FPin> FChaosOutfitAssetMakeOutfitNode::GetPinsToRemove() const
+{
+	const int32 Index = ClothAssets.Num() - 1;
+	check(ClothAssets.IsValidIndex(Index));
+	if (const FDataflowInput* const Input = FindInput(GetConnectionReference(Index)))
+	{
+		return { { UE::Dataflow::FPin::EDirection::INPUT, Input->GetType(), Input->GetName() } };
+	}
+	return Super::GetPinsToRemove();
+}
+
+void FChaosOutfitAssetMakeOutfitNode::OnPinRemoved(const UE::Dataflow::FPin& Pin)
+{
+	const int32 Index = ClothAssets.Num() - 1;
+	check(ClothAssets.IsValidIndex(Index));
+#if DO_CHECK
+	const FDataflowInput* const Input = FindInput(GetConnectionReference(Index));
+	check(Input);
+	check(Input->GetName() == Pin.Name);
+	check(Input->GetType() == Pin.Type);
+#endif
+	ClothAssets.SetNum(Index);
+
+	return Super::OnPinRemoved(Pin);
+}
+
+void FChaosOutfitAssetMakeOutfitNode::PostSerialize(const FArchive& Ar)
+{
+	// Added pins need to be restored when loading to make sure they get reconnected
+	if (Ar.IsLoading())
+	{
+		for (int32 Index = 0; Index < NumInitialClothAssets; ++Index)
+		{
+			check(FindInput(GetConnectionReference(Index)));
+		}
+
+		for (int32 Index = NumInitialClothAssets; Index < ClothAssets.Num(); ++Index)
+		{
+			FindOrRegisterInputArrayConnection(GetConnectionReference(Index));
+		}
+		if (Ar.IsTransacting())
+		{
+			const int32 OrigNumRegisteredInputs = GetNumInputs();
+			check(OrigNumRegisteredInputs >= NumRequiredInputs + NumInitialClothAssets);
+			const int32 OrigNumClothAssets = ClothAssets.Num();
+			const int32 OrigNumRegisteredClothAssets = (OrigNumRegisteredInputs - NumRequiredInputs);
+			if (OrigNumRegisteredClothAssets > OrigNumClothAssets)
+			{
+				// Inputs have been removed, temporarily expand ClothAssets so we can get connection references
+				ClothAssets.SetNum(OrigNumRegisteredClothAssets);
+				for (int32 Index = OrigNumClothAssets; Index < ClothAssets.Num(); ++Index)
+				{
+					UnregisterInputConnection(GetConnectionReference(Index));
+				}
+				ClothAssets.SetNum(OrigNumClothAssets);
+			}
+		}
+		else
+		{
+			ensureAlways(ClothAssets.Num() + NumRequiredInputs == GetNumInputs());
+		}
+	}
+}
+
+UE::Dataflow::TConnectionReference<FChaosClothAssetOrArrayAnyType> FChaosOutfitAssetMakeOutfitNode::GetConnectionReference(int32 Index) const
+{
+	return { &ClothAssets[Index], Index, &ClothAssets };
+}
+
+#undef LOCTEXT_NAMESPACE
diff --git a/Engine/Plugins/Experimental/ChaosOutfitAsset/Source/ChaosOutfitAssetDataflowNodes/Private/ChaosClothAsset/MakeOutfitNode.h b/Engine/Plugins/Experimental/ChaosOutfitAsset/Source/ChaosOutfitAssetDataflowNodes/Private/ChaosClothAsset/MakeOutfitNode.h
new file mode 100644
--- /dev/null
+++ b/Engine/Plugins/Experimental/ChaosOutfitAsset/Source/ChaosOutfitAssetDataflowNodes/Private/ChaosClothAsset/MakeOutfitNode.h
@@ -0,0 +1,52 @@
+// Copyright Epic Games, Inc. All Rights Reserved.
+
+#pragma once 
+
+#include "ChaosClothAsset/ClothAssetAnyType.h"
+#include "Dataflow/DataflowNode.h"
+#include "GeometryCollection/ManagedArrayCollection.h"
+#include "MakeOutfitNode.generated.h"
+
+class UChaosClothAssetBase;
+class UChaosOutfit;
+
+/** Add multiple cloth asset objects to an outfit collection. */
+USTRUCT(Meta = (Experimental, DataflowOutfit))
+struct FChaosOutfitAssetMakeOutfitNode final : public FDataflowNode
+{
+	GENERATED_USTRUCT_BODY()
+	DATAFLOW_NODE_DEFINE_INTERNAL(FChaosOutfitAssetMakeOutfitNode, "MakeOutfit", "Outfit", "Cloth Asset Make Outfit")
+
+public:
+	FChaosOutfitAssetMakeOutfitNode(const UE::Dataflow::FNodeParameters& InParam, FGuid InGuid = FGuid::NewGuid());
+
+private:
+	//~ Begin FDataflowNode interface
+	virtual void Evaluate(UE::Dataflow::FContext& Context, const FDataflowOutput* Out) const override;
+	virtual TArray<UE::Dataflow::FPin> AddPins() override;
+	virtual bool CanAddPin() const override
+	{
+		return true;
+	}
+	virtual bool CanRemovePin() const override
+	{
+		return ClothAssets.Num() > NumInitialClothAssets;
+	}
+	virtual TArray<UE::Dataflow::FPin> GetPinsToRemove() const override;
+	virtual void OnPinRemoved(const UE::Dataflow::FPin& Pin) override;
+	virtual void PostSerialize(const FArchive& Ar) override;
+	//~ End FDataflowNode interface
+
+	UE::Dataflow::TConnectionReference<FChaosClothAssetOrArrayAnyType> GetConnectionReference(int32 Index) const;
+
+	static constexpr int32 NumRequiredInputs = 0;
+	static constexpr int32 NumInitialClothAssets = 1;
+
+	/** The outfit output collection. */
+	UPROPERTY(Meta = (DataflowOutput))
+	TObjectPtr<const UChaosOutfit> Outfit;
+
+	/** The cloth assets to import into the collection. */
+	UPROPERTY(EditAnywhere, Category = "ClothAssetToOutfit")
+	TArray<FChaosClothAssetOrArrayAnyType> ClothAssets;
+};
diff --git a/Engine/Plugins/Experimental/ChaosOutfitAsset/Source/ChaosOutfitAssetDataflowNodes/Private/ChaosClothAsset/MergeOutfitsNode.cpp b/Engine/Plugins/Experimental/ChaosOutfitAsset/Source/ChaosOutfitAssetDataflowNodes/Private/ChaosClothAsset/MergeOutfitsNode.cpp
new file mode 100644
--- /dev/null
+++ b/Engine/Plugins/Experimental/ChaosOutfitAsset/Source/ChaosOutfitAssetDataflowNodes/Private/ChaosClothAsset/MergeOutfitsNode.cpp
@@ -0,0 +1,126 @@
+// Copyright Epic Games, Inc. All Rights Reserved.
+
+#include "ChaosClothAsset/MergeOutfitsNode.h"
+#include "ChaosClothAsset/Outfit.h"
+#include "Dataflow/DataflowInputOutput.h"
+
+#include UE_INLINE_GENERATED_CPP_BY_NAME(MergeOutfitsNode)
+
+#define LOCTEXT_NAMESPACE "ChaosOutfitAssetMergeOutfitsNode"
+
+FChaosOutfitAssetMergeOutfitsNode::FChaosOutfitAssetMergeOutfitsNode(const UE::Dataflow::FNodeParameters& InParam, FGuid InGuid)
+	: FDataflowNode(InParam, InGuid)
+{
+	check(GetNumInputs() == NumRequiredInputs);
+	for (int32 Index = 0; Index < NumInitialOptionalInputs; ++Index)
+	{
+		AddPins();
+	}
+	RegisterOutputConnection(&Outfit)
+		.SetPassthroughInput(GetConnectionReference(0));
+}
+
+void FChaosOutfitAssetMergeOutfitsNode::Evaluate(UE::Dataflow::FContext& Context, const FDataflowOutput* Out) const
+{
+	if (Out->IsA<TObjectPtr<const UChaosOutfit>>(&Outfit))
+	{
+		TObjectPtr<UChaosOutfit> OutOutfit;
+
+		for (int32 Index = 0; Index < Outfits.Num(); ++Index)
+		{
+			const TObjectPtr<const UChaosOutfit>& InOutfit = GetValue(Context, &Outfits[Index]);
+			if (InOutfit && InOutfit->GetPieces().Num())
+			{
+				if (!OutOutfit)
+				{
+					OutOutfit = NewObject<UChaosOutfit>();
+				}
+				OutOutfit->Append(*InOutfit);
+			}
+		}
+
+		if (OutOutfit)
+		{
+			SetValue<TObjectPtr<const UChaosOutfit>>(Context, OutOutfit, &Outfit);
+		}
+		else
+		{
+			SafeForwardInput(Context, GetConnectionReference(0), &Outfit);
+		}
+	}
+}
+
+TArray<UE::Dataflow::FPin> FChaosOutfitAssetMergeOutfitsNode::AddPins()
+{
+	const int32 Index = Outfits.AddDefaulted();
+	const FDataflowInput& Input = RegisterInputArrayConnection(GetConnectionReference(Index));
+	return { { UE::Dataflow::FPin::EDirection::INPUT, Input.GetType(), Input.GetName() } };
+}
+
+TArray<UE::Dataflow::FPin> FChaosOutfitAssetMergeOutfitsNode::GetPinsToRemove() const
+{
+	const int32 Index = Outfits.Num() - 1;
+	check(Outfits.IsValidIndex(Index));
+	if (const FDataflowInput* const Input = FindInput(GetConnectionReference(Index)))
+	{
+		return { { UE::Dataflow::FPin::EDirection::INPUT, Input->GetType(), Input->GetName() } };
+	}
+	return Super::GetPinsToRemove();
+}
+
+void FChaosOutfitAssetMergeOutfitsNode::OnPinRemoved(const UE::Dataflow::FPin& Pin)
+{
+	const int32 Index = Outfits.Num() - 1;
+	check(Outfits.IsValidIndex(Index));
+#if DO_CHECK
+	const FDataflowInput* const Input = FindInput(GetConnectionReference(Index));
+	check(Input);
+	check(Input->GetName() == Pin.Name);
+	check(Input->GetType() == Pin.Type);
+#endif
+	Outfits.SetNum(Index);
+
+	return Super::OnPinRemoved(Pin);
+}
+
+void FChaosOutfitAssetMergeOutfitsNode::PostSerialize(const FArchive& Ar)
+{
+	if (Ar.IsLoading())
+	{
+		for (int32 Index = 0; Index < NumInitialOptionalInputs; ++Index)
+		{
+			check(FindInput(GetConnectionReference(Index)));
+		}
+		for (int32 Index = NumInitialOptionalInputs; Index < Outfits.Num(); ++Index)
+		{
+			FindOrRegisterInputArrayConnection(GetConnectionReference(Index));
+		}
+		if (Ar.IsTransacting())
+		{
+			const int32 OrigNumRegisteredInputs = GetNumInputs();
+			check(OrigNumRegisteredInputs >= NumRequiredInputs + NumInitialOptionalInputs);
+			const int32 OrigNumCollections = Outfits.Num();
+			const int32 OrigNumRegisteredOutfits = OrigNumRegisteredInputs - NumRequiredInputs;
+			if (OrigNumRegisteredOutfits > OrigNumCollections)
+			{
+				Outfits.SetNum(GetNumInputs() - 1);
+				for (int32 Index = OrigNumCollections; Index < Outfits.Num(); ++Index)
+				{
+					UnregisterInputConnection(GetConnectionReference(Index));
+				}
+				Outfits.SetNum(OrigNumCollections);
+			}
+		}
+		else
+		{
+			ensureAlways(Outfits.Num() == GetNumInputs());
+		}
+	}
+}
+
+UE::Dataflow::TConnectionReference<TObjectPtr<const UChaosOutfit>> FChaosOutfitAssetMergeOutfitsNode::GetConnectionReference(int32 Index) const
+{
+	return { &Outfits[Index], Index, &Outfits };
+}
+
+#undef LOCTEXT_NAMESPACE
diff --git a/Engine/Plugins/Experimental/ChaosOutfitAsset/Source/ChaosOutfitAssetDataflowNodes/Private/ChaosClothAsset/MergeOutfitsNode.h b/Engine/Plugins/Experimental/ChaosOutfitAsset/Source/ChaosOutfitAssetDataflowNodes/Private/ChaosClothAsset/MergeOutfitsNode.h
new file mode 100644
--- /dev/null
+++ b/Engine/Plugins/Experimental/ChaosOutfitAsset/Source/ChaosOutfitAssetDataflowNodes/Private/ChaosClothAsset/MergeOutfitsNode.h
@@ -0,0 +1,47 @@
+// Copyright Epic Games, Inc. All Rights Reserved.
+
+#pragma once 
+
+#include "Dataflow/DataflowNode.h"
+#include "MergeOutfitsNode.generated.h"
+
+class UChaosOutfit;
+
+/** Merge multiple outfits into a single outfits. */
+USTRUCT(Meta = (Experimental, DataflowOutfit))
+struct FChaosOutfitAssetMergeOutfitsNode : public FDataflowNode
+{
+	GENERATED_USTRUCT_BODY()
+	DATAFLOW_NODE_DEFINE_INTERNAL(FChaosOutfitAssetMergeOutfitsNode, "MergeOutfits", "Outfit", "Outfit Merge Outfits")
+
+public:
+	FChaosOutfitAssetMergeOutfitsNode(const UE::Dataflow::FNodeParameters& InParam, FGuid InGuid = FGuid::NewGuid());
+
+private:
+	//~ Begin FDataflowNode interface
+	virtual void Evaluate(UE::Dataflow::FContext& Context, const FDataflowOutput* Out) const override;
+	virtual TArray<UE::Dataflow::FPin> AddPins() override;
+	virtual bool CanAddPin() const override
+	{
+		return true;
+	}
+	virtual bool CanRemovePin() const override
+	{
+		return Outfits.Num() > NumInitialOptionalInputs;
+	}
+	virtual TArray<UE::Dataflow::FPin> GetPinsToRemove() const override;
+	virtual void OnPinRemoved(const UE::Dataflow::FPin& Pin) override;
+	virtual void PostSerialize(const FArchive& Ar) override;
+	//~ End FDataflowNode interface
+
+	UE::Dataflow::TConnectionReference<TObjectPtr<const UChaosOutfit>> GetConnectionReference(int32 Index) const;
+
+	UPROPERTY()
+	TArray<TObjectPtr<const UChaosOutfit>> Outfits;
+
+	UPROPERTY(Meta = (DataflowOutput, DataflowPassthrough = "Outfits[0]"))
+	TObjectPtr<const UChaosOutfit> Outfit;
+
+	static constexpr int32 NumRequiredInputs = 0;
+	static constexpr int32 NumInitialOptionalInputs = 2;
+};
diff --git a/Engine/Plugins/Experimental/ChaosOutfitAsset/Source/ChaosOutfitAssetDataflowNodes/Private/ChaosClothAsset/OutfitAssetTerminalNode.cpp b/Engine/Plugins/Experimental/ChaosOutfitAsset/Source/ChaosOutfitAssetDataflowNodes/Private/ChaosClothAsset/OutfitAssetTerminalNode.cpp
--- a/Engine/Plugins/Experimental/ChaosOutfitAsset/Source/ChaosOutfitAssetDataflowNodes/Private/ChaosClothAsset/OutfitAssetTerminalNode.cpp
+++ b/Engine/Plugins/Experimental/ChaosOutfitAsset/Source/ChaosOutfitAssetDataflowNodes/Private/ChaosClothAsset/OutfitAssetTerminalNode.cpp
@@ -11,18 +11,17 @@
 FChaosOutfitAssetTerminalNode::FChaosOutfitAssetTerminalNode(const UE::Dataflow::FNodeParameters& InParam, FGuid InGuid)
 	: FDataflowTerminalNode(InParam, InGuid)
 {
-	RegisterInputConnection(&Collection);
+	RegisterInputConnection(&Outfit);
 }
 
 void FChaosOutfitAssetTerminalNode::SetAssetValue(TObjectPtr<UObject> Asset, UE::Dataflow::FContext& Context) const
 {
 	if (UChaosOutfitAsset* const OutfitAsset = Cast<UChaosOutfitAsset>(Asset.Get()))
 	{
-		const FManagedArrayCollection& OutfitCollection = GetValue(Context, &Collection);
-		// TODO: Process the Outfit Collection 
+		const TObjectPtr<const UChaosOutfit>& InOutfit = GetValue(Context, &Outfit);
 
-		// Rebuild the asset static data
-		OutfitAsset->Build();
+		// Build the asset
+		OutfitAsset->Build(InOutfit);
 
 		// Asset must be resaved
 		OutfitAsset->MarkPackageDirty();
diff --git a/Engine/Plugins/Experimental/ChaosOutfitAsset/Source/ChaosOutfitAssetDataflowNodes/Private/ChaosClothAsset/OutfitAssetTerminalNode.h b/Engine/Plugins/Experimental/ChaosOutfitAsset/Source/ChaosOutfitAssetDataflowNodes/Private/ChaosClothAsset/OutfitAssetTerminalNode.h
--- a/Engine/Plugins/Experimental/ChaosOutfitAsset/Source/ChaosOutfitAssetDataflowNodes/Private/ChaosClothAsset/OutfitAssetTerminalNode.h
+++ b/Engine/Plugins/Experimental/ChaosOutfitAsset/Source/ChaosOutfitAssetDataflowNodes/Private/ChaosClothAsset/OutfitAssetTerminalNode.h
@@ -3,9 +3,10 @@
 #pragma once 
 
 #include "Dataflow/DataflowTerminalNode.h"
-#include "GeometryCollection/ManagedArrayCollection.h"
 #include "OutfitAssetTerminalNode.generated.h"
 
+class UChaosOutfit;
+
 /**
  * Outfit terminal node to generate an outfit asset from several cloth assets.
  */
@@ -26,6 +27,6 @@ private:
 
 	/** Input outfit collection. */
 	UPROPERTY(Meta = (DataflowInput))
-	FManagedArrayCollection Collection;
+	TObjectPtr<const UChaosOutfit> Outfit;
 };
 
