CIS - Fixed Mac non unity builds Issue 764624:  error: no class named 'FLazySingleton' in the global namespace.

#rnx

diff --git a/Engine/Plugins/Experimental/ChaosOutfitAsset/Source/ChaosOutfitAssetEditor/Private/ChaosClothAsset/OutfitAssetEditorStyle.h b/Engine/Plugins/Experimental/ChaosOutfitAsset/Source/ChaosOutfitAssetEditor/Private/ChaosClothAsset/OutfitAssetEditorStyle.h
--- a/Engine/Plugins/Experimental/ChaosOutfitAsset/Source/ChaosOutfitAssetEditor/Private/ChaosClothAsset/OutfitAssetEditorStyle.h
+++ b/Engine/Plugins/Experimental/ChaosOutfitAsset/Source/ChaosOutfitAssetEditor/Private/ChaosClothAsset/OutfitAssetEditorStyle.h
@@ -3,6 +3,7 @@
 #pragma once
 
 #include "Styling/SlateStyle.h"
+#include "Misc/LazySingleton.h"
 
 namespace UE::Chaos::ClothAsset
 {
diff --git a/Engine/Plugins/Experimental/ChaosOutfitAsset/Source/ChaosOutfitAssetEditor/Private/ChaosClothAsset/OutfitEditorStyle.h b/Engine/Plugins/Experimental/ChaosOutfitAsset/Source/ChaosOutfitAssetEditor/Private/ChaosClothAsset/OutfitEditorStyle.h
--- a/Engine/Plugins/Experimental/ChaosOutfitAsset/Source/ChaosOutfitAssetEditor/Private/ChaosClothAsset/OutfitEditorStyle.h
+++ b/Engine/Plugins/Experimental/ChaosOutfitAsset/Source/ChaosOutfitAssetEditor/Private/ChaosClothAsset/OutfitEditorStyle.h
@@ -3,6 +3,7 @@
 #pragma once
 
 #include "Styling/SlateStyle.h"
+#include "Misc/LazySingleton.h"
 
 namespace UE::Chaos::ClothAsset
 {
