Chaos Cloth Asset - Updated categories priorities in details panel.

#rnx

diff --git a/Engine/Plugins/ChaosClothAsset/Source/ChaosClothAssetEngine/Public/ChaosClothAsset/ClothAsset.h b/Engine/Plugins/ChaosClothAsset/Source/ChaosClothAssetEngine/Public/ChaosClothAsset/ClothAsset.h
--- a/Engine/Plugins/ChaosClothAsset/Source/ChaosClothAssetEngine/Public/ChaosClothAsset/ClothAsset.h
+++ b/Engine/Plugins/ChaosClothAsset/Source/ChaosClothAssetEngine/Public/ChaosClothAsset/ClothAsset.h
@@ -38,7 +38,7 @@ ENUM_CLASS_FLAGS(EClothAssetAsyncProperties);
 /**
  * Cloth asset for pattern based simulation.
  */
-UCLASS(hidecategories = Object, BlueprintType)
+UCLASS(hidecategories = Object, BlueprintType, PrioritizeCategories = ("Dataflow"))
 class CHAOSCLOTHASSETENGINE_API UChaosClothAsset : public UChaosClothAssetBase
 {
 	GENERATED_BODY()
diff --git a/Engine/Plugins/Experimental/ChaosOutfitAsset/Source/ChaosOutfitAssetEngine/Public/ChaosClothAsset/OutfitAsset.h b/Engine/Plugins/Experimental/ChaosOutfitAsset/Source/ChaosOutfitAssetEngine/Public/ChaosClothAsset/OutfitAsset.h
--- a/Engine/Plugins/Experimental/ChaosOutfitAsset/Source/ChaosOutfitAssetEngine/Public/ChaosClothAsset/OutfitAsset.h
+++ b/Engine/Plugins/Experimental/ChaosOutfitAsset/Source/ChaosOutfitAssetEngine/Public/ChaosClothAsset/OutfitAsset.h
@@ -12,7 +12,7 @@ struct FSkeletalMaterial;
 /**
  * Outfit asset for character clothing and simulation.
  */
-UCLASS(Experimental, HideCategories = Object, BlueprintType)
+UCLASS(Experimental, HideCategories = Object, BlueprintType, PrioritizeCategories = ("Dataflow"))
 class CHAOSOUTFITASSETENGINE_API UChaosOutfitAsset final : public UChaosClothAssetBase
 {
 	GENERATED_BODY()
