Chaos Outfit Asset - Added the merging of render sections. Setup the default Outfit Asset template with the resizing graph.

#rnx

#virtualized

diff --git a/Engine/Plugins/Experimental/ChaosOutfitAsset/Source/ChaosOutfitAssetDataflowNodes/Private/ChaosClothAsset/MakeOutfitNode.h b/Engine/Plugins/Experimental/ChaosOutfitAsset/Source/ChaosOutfitAssetDataflowNodes/Private/ChaosClothAsset/MakeOutfitNode.h
--- a/Engine/Plugins/Experimental/ChaosOutfitAsset/Source/ChaosOutfitAssetDataflowNodes/Private/ChaosClothAsset/MakeOutfitNode.h
+++ b/Engine/Plugins/Experimental/ChaosOutfitAsset/Source/ChaosOutfitAssetDataflowNodes/Private/ChaosClothAsset/MakeOutfitNode.h
@@ -42,11 +42,11 @@ private:
 	static constexpr int32 NumRequiredInputs = 0;
 	static constexpr int32 NumInitialClothAssets = 1;
 
-	/** The outfit output collection. */
+	/** The outfit output. */
 	UPROPERTY(Meta = (DataflowOutput))
 	TObjectPtr<const UChaosOutfit> Outfit;
 
-	/** The cloth assets to import into the collection. */
+	/** The cloth assets or outfits to import into a new outfit. */
 	UPROPERTY(EditAnywhere, Category = "ClothAssetToOutfit")
 	TArray<FChaosClothAssetOrArrayAnyType> ClothAssets;
 };
diff --git a/Engine/Plugins/Experimental/ChaosOutfitAsset/Source/ChaosOutfitAssetDataflowNodes/Private/ChaosClothAsset/OutfitAssetTerminalNode.h b/Engine/Plugins/Experimental/ChaosOutfitAsset/Source/ChaosOutfitAssetDataflowNodes/Private/ChaosClothAsset/OutfitAssetTerminalNode.h
--- a/Engine/Plugins/Experimental/ChaosOutfitAsset/Source/ChaosOutfitAssetDataflowNodes/Private/ChaosClothAsset/OutfitAssetTerminalNode.h
+++ b/Engine/Plugins/Experimental/ChaosOutfitAsset/Source/ChaosOutfitAssetDataflowNodes/Private/ChaosClothAsset/OutfitAssetTerminalNode.h
@@ -22,7 +22,7 @@ public:
 private:
 	//~ Begin FDataflowNode interface
 	virtual void SetAssetValue(TObjectPtr<UObject> Asset, UE::Dataflow::FContext& Context) const override;
-	virtual void Evaluate(UE::Dataflow::FContext& Context, const FDataflowOutput* Out) const override {}
+	virtual void Evaluate(UE::Dataflow::FContext& Context) const override {}
 	//~ End FDataflowNode interface
 
 	/** Input outfit collection. */
diff --git a/Engine/Plugins/Experimental/ChaosOutfitAsset/Source/ChaosOutfitAssetEngine/ChaosOutfitAssetEngine.Build.cs b/Engine/Plugins/Experimental/ChaosOutfitAsset/Source/ChaosOutfitAssetEngine/ChaosOutfitAssetEngine.Build.cs
--- a/Engine/Plugins/Experimental/ChaosOutfitAsset/Source/ChaosOutfitAssetEngine/ChaosOutfitAssetEngine.Build.cs
+++ b/Engine/Plugins/Experimental/ChaosOutfitAsset/Source/ChaosOutfitAssetEngine/ChaosOutfitAssetEngine.Build.cs
@@ -26,6 +26,8 @@ public class ChaosOutfitAssetEngine : ModuleRules
 				"ChaosClothAssetEngine",
 				"ClothingSystemRuntimeCommon",
 				"Engine",
+				"RenderCore",
+				"RHI",
 			}
 		);
 
diff --git a/Engine/Plugins/Experimental/ChaosOutfitAsset/Source/ChaosOutfitAssetEngine/Private/ChaosClothAsset/ChaosOutfitAssetEngineModule.cpp b/Engine/Plugins/Experimental/ChaosOutfitAsset/Source/ChaosOutfitAssetEngine/Private/ChaosClothAsset/ChaosOutfitAssetEngineModule.cpp
--- a/Engine/Plugins/Experimental/ChaosOutfitAsset/Source/ChaosOutfitAssetEngine/Private/ChaosClothAsset/ChaosOutfitAssetEngineModule.cpp
+++ b/Engine/Plugins/Experimental/ChaosOutfitAsset/Source/ChaosOutfitAssetEngine/Private/ChaosClothAsset/ChaosOutfitAssetEngineModule.cpp
@@ -2,9 +2,12 @@
 
 #include "Modules/ModuleInterface.h"
 #include "Modules/ModuleManager.h"
+#include "ChaosClothAsset/OutfitAssetPrivate.h"
 
 #define LOCTEXT_NAMESPACE "OutfitAssetEngineModule"
 
+DEFINE_LOG_CATEGORY(LogChaosOutfitAsset);
+
 namespace UE::Chaos::ClothAsset
 {
 	class FOutfitAssetEngineModule : public IModuleInterface
diff --git a/Engine/Plugins/Experimental/ChaosOutfitAsset/Source/ChaosOutfitAssetEngine/Private/ChaosClothAsset/Outfit.cpp b/Engine/Plugins/Experimental/ChaosOutfitAsset/Source/ChaosOutfitAssetEngine/Private/ChaosClothAsset/Outfit.cpp
--- a/Engine/Plugins/Experimental/ChaosOutfitAsset/Source/ChaosOutfitAssetEngine/Private/ChaosClothAsset/Outfit.cpp
+++ b/Engine/Plugins/Experimental/ChaosOutfitAsset/Source/ChaosOutfitAssetEngine/Private/ChaosClothAsset/Outfit.cpp
@@ -4,9 +4,102 @@
 #include "Chaos/CollectionPropertyFacade.h"
 #include "ChaosClothAsset/ClothAsset.h"
 #include "ChaosClothAsset/ClothSimulationModel.h"
+#include "ChaosClothAsset/OutfitAssetPrivate.h"
+#include "Rendering/SkeletalMeshLODModel.h"
+#include "Rendering/SkeletalMeshLODRenderData.h"
+#include "Rendering/SkeletalMeshRenderData.h"
+#include "StaticMeshResources.h"
 
 #include UE_INLINE_GENERATED_CPP_BY_NAME(Outfit)
 
+struct UChaosOutfit::FLODRenderData
+{
+	// Half edge buffers are always regenerated on the last merge, a dummy buffer is used until then
+	struct FDummyHalfEdgeBuffer
+	{
+		bool bHasHalfEdges = false;
+		
+		bool IsCPUDataValid() const
+		{
+			return bHasHalfEdges;
+		}
+		void CleanUp()
+		{
+			bHasHalfEdges = false;
+		}
+		void Init(UChaosOutfit::FLODRenderData&)
+		{
+			bHasHalfEdges = true;
+		}
+	};
+
+	TArray<FSkelMeshRenderSection> RenderSections;
+	FMultiSizeIndexContainer MultiSizeIndexContainer;
+	FStaticMeshVertexBuffers StaticVertexBuffers;
+	FSkinWeightVertexBuffer SkinWeightVertexBuffer;
+	FSkeletalMeshVertexClothBuffer ClothVertexBuffer;
+	FSkinWeightProfilesData SkinWeightProfilesData;
+	TArray<FBoneIndexType> ActiveBoneIndices;
+	TArray<FBoneIndexType> RequiredBones;
+	FDummyHalfEdgeBuffer HalfEdgeBuffer;
+
+	uint32 GetNumVertices() const
+	{
+		return StaticVertexBuffers.PositionVertexBuffer.GetNumVertices();
+	}
+
+	friend FArchive& operator<<(FArchive& Ar, UChaosOutfit::FLODRenderData& LODRenderData)
+	{
+		static const FName MultiSizeIndexContainerName(TEXT("UChaosOutfit::FLODRenderData.MultiSizeIndexContainer"));
+		static const FName StaticPositionVertexBufferName(TEXT("UChaosOutfit::FLODRenderData.StaticVertexBuffers.PositionVertexBuffer"));
+		static const FName StaticMeshVertexBufferName(TEXT("UChaosOutfit::FLODRenderData.StaticVertexBuffers.StaticMeshVertexBuffer"));
+		static const FName StaticColorVertexBufferName(TEXT("UChaosOutfit::FLODRenderData.StaticVertexBuffers.ColorVertexBuffer"));
+		static const FName ClothVertexBufferName(TEXT("UChaosOutfit::FLODRenderData.ClothVertexBuffer"));
+		static const FName SourceRayTracingGeometryName(TEXT("UChaosOutfit::FLODRenderData.SourceRayTracingGeometry"));
+
+		constexpr bool bNeedsCPUAccess = true;
+
+		Ar << LODRenderData.RenderSections;
+
+		LODRenderData.MultiSizeIndexContainer.Serialize(Ar, bNeedsCPUAccess);
+		LODRenderData.MultiSizeIndexContainer.GetIndexBuffer()->SetResourceName(MultiSizeIndexContainerName);
+
+		LODRenderData.StaticVertexBuffers.PositionVertexBuffer.Serialize(Ar, bNeedsCPUAccess);
+		LODRenderData.StaticVertexBuffers.PositionVertexBuffer.SetResourceName(StaticPositionVertexBufferName);
+		LODRenderData.StaticVertexBuffers.StaticMeshVertexBuffer.Serialize(Ar, bNeedsCPUAccess);
+		LODRenderData.StaticVertexBuffers.StaticMeshVertexBuffer.SetResourceName(StaticMeshVertexBufferName);
+
+		Ar << LODRenderData.SkinWeightVertexBuffer;
+
+		Ar << LODRenderData.ClothVertexBuffer;
+		LODRenderData.ClothVertexBuffer.SetResourceName(ClothVertexBufferName);
+
+		Ar << LODRenderData.ActiveBoneIndices;
+		Ar << LODRenderData.RequiredBones;
+		Ar << LODRenderData.HalfEdgeBuffer.bHasHalfEdges;
+
+		return Ar;
+	}
+};
+
+struct UChaosOutfit::FRenderData
+{
+	TIndirectArray<FLODRenderData> LODRenderData;
+	uint8 NumInlinedLODs = 0;
+	uint8 NumNonOptionalLODs = 0;
+	bool bSupportRayTracing = false;
+
+	friend FArchive& operator<<(FArchive& Ar, UChaosOutfit::FRenderData& InRenderData)
+	{
+		Ar << InRenderData.LODRenderData;
+		Ar << InRenderData.NumInlinedLODs;
+		Ar << InRenderData.NumNonOptionalLODs;
+		Ar << InRenderData.bSupportRayTracing;
+
+		return Ar;
+	}
+};
+
 namespace UE::Chaos::ClothAsset::Private
 {
 	// Merge bone to a new skeleton and return its new index in the new skeleton
@@ -53,6 +146,44 @@ namespace UE::Chaos::ClothAsset::Private
 		}
 	}
 
+	static TArray<int32> GetChildren(const TArray<FMeshBoneInfo>& MeshBoneInfos, const int32 ParentIndex)
+	{
+		TArray<int32> Children;
+		for (int32 Index = 0; Index < MeshBoneInfos.Num(); ++Index)
+		{
+			const FMeshBoneInfo& MeshBoneInfo = MeshBoneInfos[Index];
+			if (MeshBoneInfo.ParentIndex == ParentIndex)
+			{
+				Children.Emplace(Index);
+			}
+		}
+		return Children;
+	}
+
+	static void LogHierachy(const TArray<FMeshBoneInfo>& MeshBoneInfos, const int32 ParentIndex = INDEX_NONE, int32 Indent = 0)
+	{
+		auto MakeIndentString = [](const int32 Indent) -> FString
+			{
+				FString IndentString;
+				for (int32 Index = 0; Index < Indent; ++Index)
+				{
+					IndentString += TEXT("   ");
+				}
+				return IndentString;
+			};
+
+		const FString IndentString = MakeIndentString(Indent);
+
+		for (const int32 ChildIndex : GetChildren(MeshBoneInfos, ParentIndex))
+		{
+			const FString IndentedName = IndentString + MeshBoneInfos[ChildIndex].Name.ToString();
+
+			UE_LOG(LogChaosOutfitAsset, Verbose, TEXT("%d - %s"), ChildIndex, *IndentedName);
+			
+			LogHierachy(MeshBoneInfos, ChildIndex, Indent + 1);
+		}
+	}
+
 	static void MergeSkeletons(const FReferenceSkeleton& InSkeleton, FReferenceSkeleton& InOutSkeleton, TArray<int32>& OutBoneMap)
 	{
 		const int32 NumBones = InSkeleton.GetNum();
@@ -85,6 +216,217 @@ namespace UE::Chaos::ClothAsset::Private
 			// Marked it as processed
 			BonesToProcess[BoneIndex] = false;
 		}
+
+		UE_LOG(LogChaosOutfitAsset, Verbose, TEXT("------- Outfit Reference Skseleton Merging --------"));
+		LogHierachy(InOutSkeleton.GetRefBoneInfo());
+		UE_LOG(LogChaosOutfitAsset, Verbose, TEXT("---------------------------------------------------"));
+		for (int32 Index = 0; Index < OutBoneMap.Num(); ++Index)
+		{
+			UE_LOG(LogChaosOutfitAsset, Verbose, TEXT("%d -> %d"), Index, OutBoneMap[Index]);
+		}
+		UE_LOG(LogChaosOutfitAsset, Verbose, TEXT("---------------------------------------------------"));
+	}
+
+	static TArray<uint32> GetIndices(const FMultiSizeIndexContainer& MultiSizeIndexContainer, int32 VertexOffset = 0)
+	{
+		TArray<uint32> IndexBuffer;
+		if (MultiSizeIndexContainer.IsIndexBufferValid())
+		{
+			MultiSizeIndexContainer.GetIndexBuffer(IndexBuffer);
+			if (VertexOffset)
+			{
+				for (uint32& Index : IndexBuffer)
+				{
+					Index += VertexOffset;
+				}
+			}
+		}
+		return IndexBuffer;
+	}
+
+	static TArray<FVector3f> GetPositions(const FPositionVertexBuffer& PositionVertexBuffer)
+	{
+		const uint32 NumVertices = PositionVertexBuffer.GetNumVertices();
+		TArray<FVector3f> Positions;
+		Positions.SetNumUninitialized(NumVertices);
+		for (uint32 VertexIndex = 0; VertexIndex < NumVertices; ++VertexIndex)
+		{
+			Positions[VertexIndex] = PositionVertexBuffer.VertexPosition(VertexIndex);
+		}
+		return Positions;
+	}
+
+	static TArray<FVector4f> GetTangents(const FStaticMeshVertexBuffer& StaticMeshVertexBuffer, int Axis)
+	{
+		const uint32 NumVertices = StaticMeshVertexBuffer.GetNumVertices();
+		TArray<FVector4f> Tangents;
+		Tangents.SetNumUninitialized(NumVertices);
+		for (uint32 VertexIndex = 0; VertexIndex < NumVertices; ++VertexIndex)
+		{
+			switch (Axis)
+			{
+			default: checkNoEntry();
+			case 0: Tangents[VertexIndex] = StaticMeshVertexBuffer.VertexTangentX(VertexIndex); break;
+			case 1: Tangents[VertexIndex] = StaticMeshVertexBuffer.VertexTangentY(VertexIndex); break;
+			case 2: Tangents[VertexIndex] = StaticMeshVertexBuffer.VertexTangentZ(VertexIndex); break;
+			}
+		}
+		return Tangents;
+	}
+
+	static TArray<FVector2f> GetVertexUVs(const FStaticMeshVertexBuffer& StaticMeshVertexBuffer, const uint32 MaxTexCoords)
+	{
+		const uint32 NumVertices = StaticMeshVertexBuffer.GetNumVertices();
+		const uint32 NumTexCoords = StaticMeshVertexBuffer.GetNumTexCoords();
+		check(NumTexCoords <= MaxTexCoords);
+
+		TArray<FVector2f> VertexUVs;
+		VertexUVs.SetNumUninitialized(NumVertices * MaxTexCoords);
+
+		for (uint32 VertexIndex = 0; VertexIndex < NumVertices; ++VertexIndex)
+		{
+			for (uint32 UVIndex = 0; UVIndex < NumTexCoords; ++UVIndex)
+			{
+				VertexUVs[VertexIndex * MaxTexCoords + UVIndex] = StaticMeshVertexBuffer.GetVertexUV(VertexIndex, UVIndex);
+			}
+			for (uint32 UVIndex = NumTexCoords; UVIndex < MaxTexCoords; ++UVIndex)
+			{
+				VertexUVs[VertexIndex * MaxTexCoords + UVIndex] = FVector2f::ZeroVector;
+			}
+		}
+		return VertexUVs;
+	}
+
+	static TArray<FSkinWeightInfo> GetSkinWeights(const FSkinWeightVertexBuffer& SkinWeightVertexBuffer, const bool bUse16BitBoneWeight, const TArray<int32>* const BoneMap = nullptr)
+	{
+		checkf(SkinWeightVertexBuffer.Use16BitBoneWeight() || !bUse16BitBoneWeight, TEXT("Weights can only be read from 8bit to 8bit, or 8bit to 16bit, but not 16bit to 8bit."));
+		const bool bRenormalizeTo16BitBoneWeight = bUse16BitBoneWeight && !SkinWeightVertexBuffer.Use16BitBoneWeight();
+		const uint32 MaxBoneInfluences = SkinWeightVertexBuffer.GetMaxBoneInfluences();
+		const uint32 NumVertices = SkinWeightVertexBuffer.GetNumVertices();
+		TArray<FSkinWeightInfo> SkinWeights;
+		SkinWeights.SetNumUninitialized(NumVertices);
+		for (uint32 VertexIndex = 0; VertexIndex < NumVertices; ++VertexIndex)
+		{
+			FSkinWeightInfo& SkinWeight = SkinWeights[VertexIndex];
+			if (!bRenormalizeTo16BitBoneWeight)
+			{
+				for (uint32 InfluenceIndex = 0; InfluenceIndex < MaxBoneInfluences; ++InfluenceIndex)
+				{
+					SkinWeight.InfluenceBones[InfluenceIndex] = SkinWeightVertexBuffer.GetBoneIndex(VertexIndex, InfluenceIndex);
+					SkinWeight.InfluenceWeights[InfluenceIndex] = SkinWeightVertexBuffer.GetBoneWeight(VertexIndex, InfluenceIndex);
+				}
+			}
+			else
+			{
+				uint16 TotalInfluenceWeight = 0;
+				for (uint32 InfluenceIndex = 0; InfluenceIndex < MaxBoneInfluences; ++InfluenceIndex)
+				{
+					SkinWeight.InfluenceBones[InfluenceIndex] = SkinWeightVertexBuffer.GetBoneIndex(VertexIndex, InfluenceIndex) << 8;
+					SkinWeight.InfluenceWeights[InfluenceIndex] = SkinWeightVertexBuffer.GetBoneWeight(VertexIndex, InfluenceIndex);
+					TotalInfluenceWeight += SkinWeight.InfluenceWeights[InfluenceIndex];
+				}
+				SkinWeight.InfluenceWeights[0] += TNumericLimits<uint16>::Max() - TotalInfluenceWeight;
+			}
+			// Clear the extra influences in case MaxBoneInfluences is changed
+			for (uint32 InfluenceIndex = MaxBoneInfluences; InfluenceIndex < MAX_TOTAL_INFLUENCES; ++InfluenceIndex)
+			{
+				SkinWeight.InfluenceBones[InfluenceIndex] = 0;
+				SkinWeight.InfluenceWeights[InfluenceIndex] = 0;
+			}
+		}
+		if (BoneMap)
+		{
+			for (FSkinWeightInfo& SkinWeight : SkinWeights)
+			{
+				for (uint32 InfluenceIndex = 0; InfluenceIndex < MaxBoneInfluences; ++InfluenceIndex)
+				{
+					SkinWeight.InfluenceBones[InfluenceIndex] = (*BoneMap)[SkinWeight.InfluenceBones[InfluenceIndex]];
+				}
+			}
+		}
+		return SkinWeights;
+	}
+
+	static TArray<FColor> GetVertexColors(const FStaticMeshVertexBuffers& StaticMeshVertexBuffers)
+	{
+		const uint32 NumVertices = StaticMeshVertexBuffers.PositionVertexBuffer.GetNumVertices();
+		const FColorVertexBuffer& ColorVertexBuffer = StaticMeshVertexBuffers.ColorVertexBuffer;
+		TArray<FColor> VertexColors;
+		VertexColors.SetNumUninitialized(NumVertices);
+		const bool bHasVertexColors = ColorVertexBuffer.GetAllocatedSize() != 0;
+		if (bHasVertexColors)
+		{
+			check(ColorVertexBuffer.GetNumVertices() == NumVertices);
+			for (uint32 VertexIndex = 0; VertexIndex < NumVertices; ++VertexIndex)
+			{
+				VertexColors[VertexIndex] = ColorVertexBuffer.VertexColor(VertexIndex);
+			}
+		}
+		else
+		{
+			for (FColor& VertexColor : VertexColors)
+			{
+				VertexColor = FColor::White;
+			}
+		}
+		return VertexColors;
+	}
+
+	static TArray<FMeshToMeshVertData> GetClothMappingData(const TArray<FSkelMeshRenderSection>& RenderSections)
+	{
+		TArray<FMeshToMeshVertData> MappingData;
+		int32 NumMappingData = 0;
+		for (const FSkelMeshRenderSection& RenderSection : RenderSections)
+		{
+			for (const TArray<FMeshToMeshVertData>& ClothMappingDataLOD : RenderSection.ClothMappingDataLODs)
+			{
+				NumMappingData += ClothMappingDataLOD.Num();
+			}
+		}
+		MappingData.Reserve(NumMappingData);
+		for (const FSkelMeshRenderSection& RenderSection : RenderSections)
+		{
+			for (const TArray<FMeshToMeshVertData>& ClothMappingDataLOD : RenderSection.ClothMappingDataLODs)
+			{
+				MappingData.Append(ClothMappingDataLOD);
+			}
+		}
+		return MappingData;
+	}
+
+	static TArray<FClothBufferIndexMapping> GetClothBufferIndexMappings(
+		const FSkeletalMeshVertexClothBuffer& ClothVertexBuffer,
+		const TArray<FSkelMeshRenderSection>& RenderSections,
+		const int32 VertexOffset = 0)
+	{
+		TArray<FClothBufferIndexMapping> ClothBufferIndexMappings = ClothVertexBuffer.GetClothIndexMapping();
+		if (!ClothBufferIndexMappings.Num())
+		{
+			// Add empty mapping as some of the LOD sections have clothing
+			ClothBufferIndexMappings.SetNumZeroed(RenderSections.Num());  // FClothBufferIndexMapping has no default constructor
+		}
+		else
+		{
+			// Update the mappings index/offset
+			for (FClothBufferIndexMapping& ClothBufferIndexMapping : ClothBufferIndexMappings)
+			{
+				ClothBufferIndexMapping.BaseVertexIndex += VertexOffset;
+				ClothBufferIndexMapping.MappingOffset += VertexOffset;
+				// LODBiasStride stays the same since the number of mapping for this section hasn't changed
+			}
+		}
+		return ClothBufferIndexMappings;
+	}
+
+	static void MergeBones(const FReferenceSkeleton& ReferenceSkeleton, const TArray<int32> &BoneMap, const TArray<FBoneIndexType>& BoneIndices, TArray<FBoneIndexType>& OutBoneIndices)
+	{
+		OutBoneIndices.Reserve(OutBoneIndices.Num() + BoneIndices.Num());
+		for (const uint16 BoneIndex : BoneIndices)
+		{
+			OutBoneIndices.AddUnique(BoneMap[BoneIndex]);
+		}
+		ReferenceSkeleton.EnsureParentsExistAndSort(OutBoneIndices);
+		OutBoneIndices.Shrink();
 	}
 }  // namespace UE::Chaos::ClothAsset::Private
 
@@ -95,24 +437,41 @@ FChaosOutfitPiece::FChaosOutfitPiece()
 
 FChaosOutfitPiece::~FChaosOutfitPiece() = default;
 
-FChaosOutfitPiece::FChaosOutfitPiece(const UChaosClothAssetBase& ClothAssetBase, int32 ModelIndex)
-	: Name(ClothAssetBase.GetClothSimulationModelName(ModelIndex))
-	, PhysicsAsset(ClothAssetBase.GetPhysicsAssetForModel(ModelIndex))
-	, ClothSimulationModel(MakeShared<FChaosClothSimulationModel>(*ClothAssetBase.GetClothSimulationModel(ModelIndex)))  // Deep copy model from one TSharedRef to the other
+FChaosOutfitPiece::FChaosOutfitPiece(
+	FName InName,
+	FGuid InAssetGuid,
+	const UPhysicsAsset* InPhysicsAsset,
+	const FChaosClothSimulationModel& InClothSimulationModel,
+	const TArray<TSharedRef<const FManagedArrayCollection>>& InCollections)
+	: Name(InName)
+	, AssetGuid(InAssetGuid)
+	, PhysicsAsset(InPhysicsAsset)
+	, ClothSimulationModel(MakeShared<FChaosClothSimulationModel>(InClothSimulationModel))
 {
-	DeepCopyCollections(ClothAssetBase.GetCollections(ModelIndex));
+	DeepCopyCollections(InCollections);
 }
 
+FChaosOutfitPiece::FChaosOutfitPiece(const UChaosClothAssetBase& ClothAssetBase, int32 ModelIndex)
+	: FChaosOutfitPiece(
+		ClothAssetBase.GetClothSimulationModelName(ModelIndex),
+		ClothAssetBase.GetAssetGuid(ModelIndex),
+		ClothAssetBase.GetPhysicsAssetForModel(ModelIndex),
+		*ClothAssetBase.GetClothSimulationModel(ModelIndex),
+		ClothAssetBase.GetCollections(ModelIndex))
+{}
+
 FChaosOutfitPiece::FChaosOutfitPiece(const FChaosOutfitPiece& Other)
-	: Name(Other.Name)
-	, PhysicsAsset(Other.PhysicsAsset)
-	, ClothSimulationModel(MakeShared<FChaosClothSimulationModel>(*Other.ClothSimulationModel))
-{
-	DeepCopyCollections(Other.Collections);
-}
+	: FChaosOutfitPiece(
+		Other.Name,
+		Other.AssetGuid,
+		Other.PhysicsAsset,
+		*Other.ClothSimulationModel,
+		Other.Collections)
+{}
 
 FChaosOutfitPiece::FChaosOutfitPiece(FChaosOutfitPiece&& Other)
 	: Name(MoveTemp(Other.Name))
+	, AssetGuid(MoveTemp(Other.AssetGuid))
 	, PhysicsAsset(MoveTemp(Other.PhysicsAsset))
 	, ClothSimulationModel(MoveTemp(Other.ClothSimulationModel))
 	, Collections(MoveTemp(Other.Collections))
@@ -124,6 +483,7 @@ FChaosOutfitPiece& FChaosOutfitPiece::operator=(const FChaosOutfitPiece& Other)
 	{
 		Name = Other.Name;
 		PhysicsAsset = Other.PhysicsAsset;
+		AssetGuid = Other.AssetGuid;
 		*ClothSimulationModel = *Other.ClothSimulationModel;
 		DeepCopyCollections(Other.Collections);
 	}
@@ -136,6 +496,7 @@ FChaosOutfitPiece& FChaosOutfitPiece::operator=(FChaosOutfitPiece&& Other)
 	{
 		Name = MoveTemp(Other.Name);
 		PhysicsAsset = MoveTemp(Other.PhysicsAsset);
+		AssetGuid = MoveTemp(Other.AssetGuid);
 		ClothSimulationModel = MoveTemp(Other.ClothSimulationModel);
 		Collections = MoveTemp(Other.Collections);
 	}
@@ -195,10 +556,104 @@ void FChaosOutfitPiece::RemapBoneIndices(const TArray<int32>& BoneMap)
 UChaosOutfit::UChaosOutfit(const FObjectInitializer& ObjectInitializer)
 	: Super(ObjectInitializer)
 {
-	// Create a default reference skeleton
-	ReferenceSkeleton.Empty(1);
-	FReferenceSkeletonModifier ReferenceSkeletonModifier(ReferenceSkeleton, nullptr);
+	Init(
+		Pieces,
+		ReferenceSkeleton,
+		RenderData,
+		Materials);
+}
+
+UChaosOutfit::UChaosOutfit(FVTableHelper& Helper)
+	: Super(Helper)
+{
+}
+
+UChaosOutfit::~UChaosOutfit() = default;
 
+void UChaosOutfit::Append(const UChaosOutfit& Other)
+{
+	Merge(
+		Other.GetReferenceSkeleton(),
+		Other.RenderData.Get(),
+		Other.Materials,
+		Other.Pieces.Num(),
+		[this, &Other](int32 Index) -> FChaosOutfitPiece
+		{
+			return FChaosOutfitPiece(Other.Pieces[Index]);
+		},
+		Pieces,
+		ReferenceSkeleton,
+		RenderData.Get(),
+		Materials);
+}
+
+void UChaosOutfit::Add(const UChaosClothAssetBase& ClothAssetBase)
+{
+	Merge(
+		ClothAssetBase.GetRefSkeleton(),
+		ClothAssetBase.GetResourceForRendering(),
+		ClothAssetBase.GetMaterials(),
+		ClothAssetBase.GetNumClothSimulationModels(),
+		[this, &ClothAssetBase](int32 Index) -> FChaosOutfitPiece
+		{
+			return FChaosOutfitPiece(ClothAssetBase, Index);
+		},
+		Pieces,
+		ReferenceSkeleton,
+		RenderData.Get(),
+		Materials);
+}
+
+void UChaosOutfit::CopyTo(
+	TArray<FChaosOutfitPiece>& OutPieces,
+	FReferenceSkeleton& OutReferenceSkeleton,
+	TUniquePtr<FSkeletalMeshRenderData>& OutSkeletalMeshRenderData,
+	TArray<FSkeletalMaterial>& OutMaterials) const
+{
+	Init(
+		OutPieces,
+		OutReferenceSkeleton,
+		OutSkeletalMeshRenderData,
+		OutMaterials);
+
+	Merge(
+		ReferenceSkeleton,
+		RenderData.Get(),
+		Materials,
+		Pieces.Num(),
+		[this](int32 Index) -> FChaosOutfitPiece
+		{
+			return FChaosOutfitPiece(Pieces[Index]);
+		},
+		OutPieces,
+		OutReferenceSkeleton,
+		OutSkeletalMeshRenderData.Get(),
+		OutMaterials);
+}
+
+void UChaosOutfit::Serialize(FArchive& Ar)
+{
+	Super::Serialize(Ar);
+
+	Ar << ReferenceSkeleton;
+	Ar << *RenderData;
+}
+
+template<typename RenderDataType UE_REQUIRES_DEFINITION(std::is_same_v<RenderDataType, FSkeletalMeshRenderData> || std::is_same_v<RenderDataType, UChaosOutfit::FRenderData>)>
+void UChaosOutfit::Init(
+	TArray<FChaosOutfitPiece>& OutPieces,
+	FReferenceSkeleton& OutReferenceSkeleton,
+	TUniquePtr<RenderDataType>& OutSkeletalMeshRenderData,
+	TArray<FSkeletalMaterial>& OutMaterials)
+{
+	OutPieces.Reset();
+	OutReferenceSkeleton.Empty();
+	OutSkeletalMeshRenderData = MakeUnique<RenderDataType>();
+	OutMaterials.Reset();
+
+	// Create a default valid reference skeleton
+	OutReferenceSkeleton.Empty(1);
+	FReferenceSkeletonModifier ReferenceSkeletonModifier(OutReferenceSkeleton, nullptr);
 	FMeshBoneInfo MeshBoneInfo;
 	constexpr const TCHAR* RootName = TEXT("Root");
 	MeshBoneInfo.ParentIndex = INDEX_NONE;
@@ -207,46 +662,275 @@ UChaosOutfit::UChaosOutfit(const FObjectInitializer& ObjectInitializer)
 #endif
 	MeshBoneInfo.Name = FName(RootName);
 	ReferenceSkeletonModifier.Add(MeshBoneInfo, FTransform::Identity);
+
+	// Set default values for the empty render data that can be usefully used in a merge
+	OutSkeletalMeshRenderData->bSupportRayTracing = false;
 }
 
-UChaosOutfit::UChaosOutfit(FVTableHelper& Helper)
-	: Super(Helper)
+template<typename InLODRenderDataType, typename OutLODRenderDataType>
+void UChaosOutfit::MergeLODRenderDatas(
+	const InLODRenderDataType& LODRenderData,
+	const TArray<FGuid>& AssetGuids,
+	const int32 MaterialOffset,
+	const FReferenceSkeleton& ReferenceSkeleton,
+	const TArray<int32>& BoneMap,
+	OutLODRenderDataType& OutLODRenderData)
 {
-}
+	using namespace UE::Chaos::ClothAsset::Private;
 
-UChaosOutfit::~UChaosOutfit() = default;
+	const uint32 VertexOffset = OutLODRenderData.GetNumVertices();
+	const uint32 NumVertices = VertexOffset + LODRenderData.GetNumVertices();
 
-void UChaosOutfit::Append(const UChaosOutfit& Other)
-{
-	TArray<int32> BoneMap;
-	UE::Chaos::ClothAsset::Private::MergeSkeletons(Other.GetReferenceSkeleton(), ReferenceSkeleton, BoneMap);
+	// Merge index buffer
+	TArray<uint32> IndexBuffer = GetIndices(OutLODRenderData.MultiSizeIndexContainer);
+	const uint32 IndexOffset = IndexBuffer.Num();
+
+	IndexBuffer.Append(GetIndices(LODRenderData.MultiSizeIndexContainer, VertexOffset));
+	const uint32 NumIndices = IndexBuffer.Num();
+
+	const uint8 DataTypeSize = (NumVertices < (uint32)TNumericLimits<uint16>::Max()) ? sizeof(uint16) : sizeof(uint32);
+	OutLODRenderData.MultiSizeIndexContainer.RebuildIndexBuffer(DataTypeSize, IndexBuffer);
 
-	Pieces.Reserve(Pieces.Num() + Other.Pieces.Num());
+	// Merge positions
+	TArray<FVector3f> Positions = GetPositions(OutLODRenderData.StaticVertexBuffers.PositionVertexBuffer);
+	Positions.Append(GetPositions(LODRenderData.StaticVertexBuffers.PositionVertexBuffer));
+	check(Positions.Num() == NumVertices);
 
-	for (const FChaosOutfitPiece& Piece : Other.Pieces)
+	// Merge tangents
+	const bool bUseHighPrecisionTangentBasis =
+		OutLODRenderData.StaticVertexBuffers.StaticMeshVertexBuffer.GetUseHighPrecisionTangentBasis() ||
+		LODRenderData.StaticVertexBuffers.StaticMeshVertexBuffer.GetUseHighPrecisionTangentBasis();
+
+	TArray<FVector4f> Tangents[3];
+	for (int Axis = 0; Axis < 3; ++Axis)
 	{
-		FChaosOutfitPiece& NewPiece = Pieces.Emplace_GetRef(Piece);
-		NewPiece.RemapBoneIndices(BoneMap);
+		Tangents[Axis] = GetTangents(OutLODRenderData.StaticVertexBuffers.StaticMeshVertexBuffer, Axis);
+		Tangents[Axis].Append(GetTangents(LODRenderData.StaticVertexBuffers.StaticMeshVertexBuffer, Axis));
+		check(Tangents[Axis].Num() == NumVertices);
 	}
+
+	// Merge UVs
+	const uint32 MaxTexCoords = FMath::Max(
+		LODRenderData.StaticVertexBuffers.StaticMeshVertexBuffer.GetNumTexCoords(),
+		OutLODRenderData.StaticVertexBuffers.StaticMeshVertexBuffer.GetNumTexCoords());
+	const bool bUseFullPrecisionUVs =
+		OutLODRenderData.StaticVertexBuffers.StaticMeshVertexBuffer.GetUseFullPrecisionUVs() ||
+		LODRenderData.StaticVertexBuffers.StaticMeshVertexBuffer.GetUseFullPrecisionUVs();
+
+	TArray<FVector2f> VertexUVs = GetVertexUVs(OutLODRenderData.StaticVertexBuffers.StaticMeshVertexBuffer, MaxTexCoords);
+	VertexUVs.Append(GetVertexUVs(LODRenderData.StaticVertexBuffers.StaticMeshVertexBuffer, MaxTexCoords));
+	check(VertexUVs.Num() == NumVertices * MaxTexCoords);
+
+	// Merge Vertex attributes
+
+	// TODO: 
+	//     The Cloth Asset doesn't have any vertex attributes,
+	//     but if it had, merging them is non trivial due to no access to the names
+
+	// Merge Vertex colors
+	const bool bHasVertexColors = NumVertices && (
+		OutLODRenderData.StaticVertexBuffers.ColorVertexBuffer.GetAllocatedSize() != 0 ||
+		LODRenderData.StaticVertexBuffers.ColorVertexBuffer.GetAllocatedSize() != 0);
+
+	TArray<FColor> VertexColors;
+	if (bHasVertexColors)
+	{
+		VertexColors = GetVertexColors(OutLODRenderData.StaticVertexBuffers);
+		VertexColors.Append(GetVertexColors(LODRenderData.StaticVertexBuffers));
+		check(VertexColors.Num() == NumVertices);
+	}
+
+	// Init vertex buffers
+	FPositionVertexBuffer& PositionVertexBuffer = OutLODRenderData.StaticVertexBuffers.PositionVertexBuffer;
+	FStaticMeshVertexBuffer& StaticMeshVertexBuffer = OutLODRenderData.StaticVertexBuffers.StaticMeshVertexBuffer;
+
+	StaticMeshVertexBuffer.SetUseFullPrecisionUVs(bUseFullPrecisionUVs);
+	StaticMeshVertexBuffer.SetUseHighPrecisionTangentBasis(bUseHighPrecisionTangentBasis);
+
+	PositionVertexBuffer.Init(NumVertices);
+	StaticMeshVertexBuffer.Init(NumVertices, MaxTexCoords);
+
+	for (uint32 VertexIndex = 0; VertexIndex < NumVertices; ++VertexIndex)
+	{
+		PositionVertexBuffer.VertexPosition(VertexIndex) = Positions[VertexIndex];
+		StaticMeshVertexBuffer.SetVertexTangents(VertexIndex, Tangents[0][VertexIndex], Tangents[1][VertexIndex], Tangents[2][VertexIndex]);
+		for (uint32 UVIndex = 0; UVIndex < MaxTexCoords; ++UVIndex)
+		{
+			StaticMeshVertexBuffer.SetVertexUV(VertexIndex, UVIndex, VertexUVs[VertexIndex * MaxTexCoords + UVIndex]);
+		}
+	}
+	if (bHasVertexColors)
+	{
+		FColorVertexBuffer& ColorVertexBuffer = OutLODRenderData.StaticVertexBuffers.ColorVertexBuffer;
+		ColorVertexBuffer.InitFromColorArray(VertexColors.GetData(), NumVertices);
+	}
+
+	// Skinweight buffer
+	const bool bNeedsCPUAccess =
+		OutLODRenderData.SkinWeightVertexBuffer.GetNeedsCPUAccess() ||
+		LODRenderData.SkinWeightVertexBuffer.GetNeedsCPUAccess();
+	const uint32 MaxBoneInfluences = FMath::Max(
+		OutLODRenderData.SkinWeightVertexBuffer.GetMaxBoneInfluences(),
+		LODRenderData.SkinWeightVertexBuffer.GetMaxBoneInfluences());
+	const bool bUse16BitBoneIndex =
+		OutLODRenderData.SkinWeightVertexBuffer.Use16BitBoneIndex() ||
+		LODRenderData.SkinWeightVertexBuffer.Use16BitBoneIndex();
+	const bool bUse16BitBoneWeight =
+		OutLODRenderData.SkinWeightVertexBuffer.Use16BitBoneWeight() ||
+		LODRenderData.SkinWeightVertexBuffer.Use16BitBoneWeight();
+
+	TArray<FSkinWeightInfo> SkinWeights = GetSkinWeights(OutLODRenderData.SkinWeightVertexBuffer, bUse16BitBoneWeight);
+	SkinWeights.Append(GetSkinWeights(LODRenderData.SkinWeightVertexBuffer, bUse16BitBoneWeight, &BoneMap));
+
+	OutLODRenderData.SkinWeightVertexBuffer.SetNeedsCPUAccess(bNeedsCPUAccess);
+	OutLODRenderData.SkinWeightVertexBuffer.SetMaxBoneInfluences(MaxBoneInfluences);
+	OutLODRenderData.SkinWeightVertexBuffer.SetUse16BitBoneIndex(bUse16BitBoneIndex);
+	OutLODRenderData.SkinWeightVertexBuffer.SetUse16BitBoneWeight(bUse16BitBoneWeight);
+	OutLODRenderData.SkinWeightVertexBuffer = SkinWeights;  // Assigns the skinweights
+
+	// Skinweight profiles
+	OutLODRenderData.SkinWeightProfilesData.Init(&OutLODRenderData.SkinWeightVertexBuffer);
+
+	// TODO: 
+	//     The Cloth Asset doesn't have SkinWeightProfilesData yet, 
+	//     but the merging code will have to be implemented once it does.
+
+	// Half edges
+	const bool bHasHalfEdges = NumVertices && (
+		OutLODRenderData.HalfEdgeBuffer.IsCPUDataValid() ||
+		LODRenderData.HalfEdgeBuffer.IsCPUDataValid());
+
+	if (bHasHalfEdges)
+	{
+		OutLODRenderData.HalfEdgeBuffer.CleanUp();
+		OutLODRenderData.HalfEdgeBuffer.Init(OutLODRenderData);
+	}
+
+	// Merge cloth data
+	const bool bHasClothingData = NumVertices && (
+		OutLODRenderData.ClothVertexBuffer.GetNumVertices() != 0 ||  // Note GetNumVertices() doesn't return NumVertices because
+		LODRenderData.ClothVertexBuffer.GetNumVertices() != 0);      // of the LODBias there could be several mappings per vertex
+
+	if (bHasClothingData)
+	{
+		TArray<FMeshToMeshVertData> ClothMappingData = GetClothMappingData(OutLODRenderData.RenderSections);
+		ClothMappingData.Append(GetClothMappingData(LODRenderData.RenderSections));
+
+		TArray<FClothBufferIndexMapping> ClothIndexMappings =
+			GetClothBufferIndexMappings(OutLODRenderData.ClothVertexBuffer, OutLODRenderData.RenderSections);
+		ClothIndexMappings.Append(
+			GetClothBufferIndexMappings(LODRenderData.ClothVertexBuffer, LODRenderData.RenderSections, VertexOffset));
+
+		OutLODRenderData.ClothVertexBuffer.Init(ClothMappingData, ClothIndexMappings);
+	}
+
+	// Merge sections
+	const int32 SectionOffset = OutLODRenderData.RenderSections.Num();  // The index of the first merged section
+	const int32 NumSections = SectionOffset + LODRenderData.RenderSections.Num();
+	OutLODRenderData.RenderSections.Reserve(NumSections);
+
+	for (const FSkelMeshRenderSection& RenderSection : LODRenderData.RenderSections)
+	{
+		const int32 PieceIndex = AssetGuids.Find(RenderSection.ClothingData.AssetGuid);
+
+		FSkelMeshRenderSection& OutRenderSection = OutLODRenderData.RenderSections.AddDefaulted_GetRef();
+
+		OutRenderSection.MaterialIndex = RenderSection.MaterialIndex + MaterialOffset;
+		OutRenderSection.BaseIndex = RenderSection.BaseIndex + IndexOffset;
+		OutRenderSection.NumTriangles = RenderSection.NumTriangles;
+		OutRenderSection.bRecomputeTangent = RenderSection.bRecomputeTangent;
+		OutRenderSection.bCastShadow = RenderSection.bCastShadow;
+		OutRenderSection.bVisibleInRayTracing = RenderSection.bVisibleInRayTracing;
+		OutRenderSection.RecomputeTangentsVertexMaskChannel = RenderSection.RecomputeTangentsVertexMaskChannel;
+		OutRenderSection.BaseVertexIndex = RenderSection.BaseVertexIndex + VertexOffset;
+		OutRenderSection.ClothMappingDataLODs = RenderSection.ClothMappingDataLODs;
+		OutRenderSection.BoneMap = RenderSection.BoneMap;
+		OutRenderSection.NumVertices = RenderSection.NumVertices;
+		OutRenderSection.MaxBoneInfluences = RenderSection.MaxBoneInfluences;
+		OutRenderSection.CorrespondClothAssetIndex = PieceIndex;
+		OutRenderSection.ClothingData = RenderSection.ClothingData;
+		OutRenderSection.DuplicatedVerticesBuffer.Init(0, TMap<int32, TArray<int32>>());
+		OutRenderSection.bDisabled = RenderSection.bDisabled;
+
+		for (FBoneIndexType& BoneIndex : OutRenderSection.BoneMap)
+		{
+			BoneIndex = BoneMap[BoneIndex];
+		}
+	}
+
+	// Update used bone indices
+	MergeBones(ReferenceSkeleton, BoneMap, LODRenderData.ActiveBoneIndices, OutLODRenderData.ActiveBoneIndices);
+	MergeBones(ReferenceSkeleton, BoneMap, LODRenderData.RequiredBones, OutLODRenderData.RequiredBones);
 }
 
-void UChaosOutfit::Add(const UChaosClothAssetBase& ClothAssetBase)
+template<typename InRenderDataType, typename OutRenderDataType>
+void UChaosOutfit::Merge(
+	const FReferenceSkeleton& InReferenceSkeleton,
+	const InRenderDataType* const InSkeletalMeshRenderData,
+	const TArray<FSkeletalMaterial>& InMaterials,
+	const int32 NumPieces,
+	TFunction<FChaosOutfitPiece(int32)>&& MakePiece,
+	TArray<FChaosOutfitPiece>& OutPieces,
+	FReferenceSkeleton& OutReferenceSkeleton,
+	OutRenderDataType* const OutSkeletalMeshRenderData,
+	TArray<FSkeletalMaterial>& OutMaterials)
 {
+	using namespace UE::Chaos::ClothAsset::Private;
+
+	check(InSkeletalMeshRenderData);
+
+	// Merge skeletons
 	TArray<int32> BoneMap;
-	UE::Chaos::ClothAsset::Private::MergeSkeletons(ClothAssetBase.GetRefSkeleton(), ReferenceSkeleton, BoneMap);
+	MergeSkeletons(InReferenceSkeleton, OutReferenceSkeleton, BoneMap);
 
-	Pieces.Reserve(Pieces.Num() + ClothAssetBase.GetNumClothSimulationModels());
+	// Merge pieces
+	TArray<FGuid> AssetGuids;
+	AssetGuids.Reserve(NumPieces);
+	OutPieces.Reserve(OutPieces.Num() + NumPieces);
 
-	for (int32 ModelIndex = 0; ModelIndex < ClothAssetBase.GetNumClothSimulationModels(); ++ModelIndex)
+	for (int32 PieceIndex = 0; PieceIndex < NumPieces; ++PieceIndex)
 	{
-		FChaosOutfitPiece& NewPiece = Pieces.Emplace_GetRef(ClothAssetBase, ModelIndex);
-		NewPiece.RemapBoneIndices(BoneMap);
+		FChaosOutfitPiece& Piece = OutPieces.Add_GetRef(MakePiece(PieceIndex));
+		Piece.RemapBoneIndices(BoneMap);
+		AssetGuids.AddUnique(Piece.AssetGuid);
 	}
-}
 
-void UChaosOutfit::Serialize(FArchive& Ar)
-{
-	Super::Serialize(Ar);
+	// Merge materials
+	const int32 MaterialOffset = OutMaterials.Num();
+	OutMaterials.Append(InMaterials);
 
-	Ar << ReferenceSkeleton;
+	// Merge render data, only for existing LODs
+	// Note: This does not create LODs for the added sections when they are missing
+	//       as it would require to duplicate data and add new LODBias clothing data.
+	if (InSkeletalMeshRenderData)
+	{
+		using InLODRenderDataType = decltype(InSkeletalMeshRenderData->LODRenderData);
+		using OutLODRenderDataType = decltype(OutSkeletalMeshRenderData->LODRenderData);
+		using LODRenderDataElementType = typename OutLODRenderDataType::ElementType;
+
+		const InLODRenderDataType& InLODRenderData = InSkeletalMeshRenderData->LODRenderData;
+		OutLODRenderDataType& LODRenderData = OutSkeletalMeshRenderData->LODRenderData;
+
+		const int32 NumLODs = InLODRenderData.Num();
+		LODRenderData.Reserve(NumLODs);
+
+		for (int32 LODIndex = 0; LODIndex < NumLODs; ++LODIndex)
+		{
+			if (LODIndex >= LODRenderData.Num())
+			{
+				LODRenderData.Add(new LODRenderDataElementType());
+			}
+			MergeLODRenderDatas(
+				InLODRenderData[LODIndex],
+				AssetGuids,
+				MaterialOffset,
+				OutReferenceSkeleton,
+				BoneMap,
+				LODRenderData[LODIndex]);
+		}
+
+		OutSkeletalMeshRenderData->NumInlinedLODs = NumLODs;
+		OutSkeletalMeshRenderData->NumNonOptionalLODs = NumLODs;
+		OutSkeletalMeshRenderData->bSupportRayTracing = OutSkeletalMeshRenderData->bSupportRayTracing || InSkeletalMeshRenderData->bSupportRayTracing;
+	}
 }
diff --git a/Engine/Plugins/Experimental/ChaosOutfitAsset/Source/ChaosOutfitAssetEngine/Private/ChaosClothAsset/OutfitAsset.cpp b/Engine/Plugins/Experimental/ChaosOutfitAsset/Source/ChaosOutfitAssetEngine/Private/ChaosClothAsset/OutfitAsset.cpp
--- a/Engine/Plugins/Experimental/ChaosOutfitAsset/Source/ChaosOutfitAssetEngine/Private/ChaosClothAsset/OutfitAsset.cpp
+++ b/Engine/Plugins/Experimental/ChaosOutfitAsset/Source/ChaosOutfitAssetEngine/Private/ChaosClothAsset/OutfitAsset.cpp
@@ -2,11 +2,13 @@
 
 #include "ChaosClothAsset/OutfitAsset.h"
 #include "ChaosClothAsset/ClothSimulationModel.h"
+#include "Engine/SkinnedAssetCommon.h"
 #if WITH_EDITOR
 #include "DerivedDataCacheInterface.h"
 #endif
 #include "Rendering/SkeletalMeshRenderData.h"
 
+
 #include UE_INLINE_GENERATED_CPP_BY_NAME(OutfitAsset)
 
 // If Chaos outfit asset derived data needs to be rebuilt (new format, serialization differences, etc.) replace the version GUID below with a new one.
@@ -19,6 +21,9 @@ UChaosOutfitAsset::UChaosOutfitAsset(const FObjectInitializer& ObjectInitializer
 PRAGMA_DISABLE_DEPRECATION_WARNINGS
 	DataflowInstance.SetDataflowTerminal(TEXT("OutfitAssetTerminal"));
 PRAGMA_ENABLE_DEPRECATION_WARNINGS
+
+	// Init an empty asset with a root bone and empty but initialized render data
+	Build(nullptr);
 }
 
 UChaosOutfitAsset::UChaosOutfitAsset(FVTableHelper& Helper)
@@ -40,6 +45,93 @@ bool UChaosOutfitAsset::HasValidClothSimulationModels() const
 	return false;
 }
 
+void UChaosOutfitAsset::Build(const TObjectPtr<const UChaosOutfit> Outfit)
+{
+	// Stop the rendering
+	ReleaseResources();
+
+	// Copy the outfit to this asset
+	TUniquePtr<FSkeletalMeshRenderData> TempSkeletalMeshRenderData = MakeUnique<FSkeletalMeshRenderData>();
+	FReferenceSkeleton TempReferenceSkeleton;
+
+	if (Outfit)
+	{
+		Outfit->CopyTo(
+			Pieces,
+			TempReferenceSkeleton,
+			TempSkeletalMeshRenderData,
+			Materials);
+	}
+	else
+	{
+		UChaosOutfit::Init(
+			Pieces,
+			TempReferenceSkeleton,
+			TempSkeletalMeshRenderData,
+			Materials);
+	}
+
+	// Set the new reference skeleton
+	SetReferenceSkeleton(&TempReferenceSkeleton);
+	CalculateInvRefMatrices();
+
+	// Create the render data
+	if (!TempSkeletalMeshRenderData->LODRenderData.Num())
+	{
+		TempSkeletalMeshRenderData->LODRenderData.Add(new FSkeletalMeshLODRenderData());  // Always needs at least one LOD when rendering
+		TempSkeletalMeshRenderData->LODRenderData[0].StaticVertexBuffers.PositionVertexBuffer.Init(0);  //  Required for serialization
+		TempSkeletalMeshRenderData->LODRenderData[0].StaticVertexBuffers.StaticMeshVertexBuffer.Init(0, 0);  //  Required for serialization
+	}
+	SetResourceForRendering(MoveTemp(TempSkeletalMeshRenderData));
+
+	CalculateBounds();
+
+	const int32 NumLODs = GetResourceForRendering()->LODRenderData.Num();  // The render data will always look for at least one default LOD 0
+	LODInfo.Reset(NumLODs);
+	LODInfo.AddDefaulted(NumLODs);
+
+	if (FApp::CanEverRender())
+	{
+		InitResources();
+	}
+
+	// Re-register any components using this asset to restart the simulation and update the leader bone map
+	ReregisterComponents();
+}
+
+void UChaosOutfitAsset::Serialize(FArchive& Ar)
+{
+	LLM_SCOPE_BYNAME(TEXT("Physics/Cloth"));
+	Super::Serialize(Ar);
+
+	{
+		LLM_SCOPE_BYNAME(TEXT("Physics/ClothRendering"));
+		if (Ar.IsLoading())
+		{
+			SetResourceForRendering(MakeUnique<FSkeletalMeshRenderData>());
+		}
+		GetResourceForRendering()->Serialize(Ar, this);
+	}
+}
+
+void UChaosOutfitAsset::PostLoad()
+{
+	LLM_SCOPE_BYNAME(TEXT("Physics/Cloth"));
+	Super::PostLoad();
+
+	if (FApp::CanEverRender())
+	{
+		InitResources();
+	}
+	else
+	{
+		UpdateUVChannelData(false);  // Update any missing data when cooking
+	}
+
+	CalculateInvRefMatrices();
+	CalculateBounds();
+}
+
 FName UChaosOutfitAsset::GetClothSimulationModelName(int32 ModelIndex) const
 {
 	return Pieces[ModelIndex].Name;
@@ -62,7 +154,7 @@ const UPhysicsAsset* UChaosOutfitAsset::GetPhysicsAssetForModel(int32 ModelIndex
 
 FGuid UChaosOutfitAsset::GetAssetGuid(int32 ModelIndex) const
 {
-	return FGuid(); // Pieces[ModelIndex].AssetGuid;
+	return Pieces[ModelIndex].AssetGuid;
 }
 
 #if WITH_EDITOR
@@ -77,19 +169,18 @@ FString UChaosOutfitAsset::BuildDerivedDataKey(const ITargetPlatform* TargetPlat
 }
 #endif // #if WITH_EDITOR
 
-void UChaosOutfitAsset::Build(const TObjectPtr<const UChaosOutfit> Outfit)
+void UChaosOutfitAsset::CalculateBounds()
 {
-	// Copy pieces across
-	const TArray<FChaosOutfitPiece>& InPieces = Outfit->GetPieces();
-	Pieces.Reset(InPieces.Num());
-	Pieces.Append(InPieces);
-
-	// Copy reference skeleton
-	SetReferenceSkeleton(&Outfit->GetReferenceSkeleton());
-
-	// Create dummy skm render data
-	SetResourceForRendering(MakeUnique<FSkeletalMeshRenderData>());
+	FBox BoundingBox(ForceInit);
+	for (const FSkeletalMeshLODRenderData& LODRenderDatum : GetResourceForRendering()->LODRenderData)
+	{
+		const FPositionVertexBuffer& PositionVertexBuffer = LODRenderDatum.StaticVertexBuffers.PositionVertexBuffer;
 
-	// Re-register any components using this asset to restart the simulation and update the leader bone map
-	ReregisterComponents();
+		for (uint32 VertexIndex = 0; VertexIndex < PositionVertexBuffer.GetNumVertices(); ++VertexIndex)
+		{
+			const FVector3f& VertexPosition = PositionVertexBuffer.VertexPosition(VertexIndex);
+			BoundingBox += (FVector)VertexPosition;
+		}
+	}
+	Bounds = FBoxSphereBounds(BoundingBox);
 }
diff --git a/Engine/Plugins/Experimental/ChaosOutfitAsset/Source/ChaosOutfitAssetEngine/Private/ChaosClothAsset/OutfitAssetPrivate.h b/Engine/Plugins/Experimental/ChaosOutfitAsset/Source/ChaosOutfitAssetEngine/Private/ChaosClothAsset/OutfitAssetPrivate.h
new file mode 100644
--- /dev/null
+++ b/Engine/Plugins/Experimental/ChaosOutfitAsset/Source/ChaosOutfitAssetEngine/Private/ChaosClothAsset/OutfitAssetPrivate.h
@@ -0,0 +1,9 @@
+// Copyright Epic Games, Inc. All Rights Reserved.
+
+#pragma once
+
+#include "Logging/LogMacros.h"
+#include "Stats/Stats.h"
+
+DECLARE_LOG_CATEGORY_EXTERN(LogChaosOutfitAsset, Log, All);
+DECLARE_STATS_GROUP(TEXT("ChaosOutfitAsset"), STATGROUP_ChaosOutfitAsset, STATCAT_Advanced);
diff --git a/Engine/Plugins/Experimental/ChaosOutfitAsset/Source/ChaosOutfitAssetEngine/Public/ChaosClothAsset/Outfit.h b/Engine/Plugins/Experimental/ChaosOutfitAsset/Source/ChaosOutfitAssetEngine/Public/ChaosClothAsset/Outfit.h
--- a/Engine/Plugins/Experimental/ChaosOutfitAsset/Source/ChaosOutfitAssetEngine/Public/ChaosClothAsset/Outfit.h
+++ b/Engine/Plugins/Experimental/ChaosOutfitAsset/Source/ChaosOutfitAssetEngine/Public/ChaosClothAsset/Outfit.h
@@ -10,9 +10,11 @@
 #include "Outfit.generated.h"
 
 class UChaosClothAssetBase;
+class FSkeletalMeshRenderData;
 class UPhysicsAsset;
 struct FChaosClothSimulationModel;
 struct FManagedArrayCollection;
+struct FSkeletalMaterial;
 
 USTRUCT()
 struct CHAOSOUTFITASSETENGINE_API FChaosOutfitPiece final
@@ -22,6 +24,9 @@ struct CHAOSOUTFITASSETENGINE_API FChaosOutfitPiece final
 	UPROPERTY()
 	FName Name;
 
+	UPROPERTY()
+	FGuid AssetGuid;
+
 	UPROPERTY()
 	TObjectPtr<const UPhysicsAsset> PhysicsAsset;
 
@@ -31,15 +36,22 @@ struct CHAOSOUTFITASSETENGINE_API FChaosOutfitPiece final
 
 	FChaosOutfitPiece();
 	~FChaosOutfitPiece();
+	FChaosOutfitPiece(
+		FName InName, 
+		FGuid InAssetGuid, 
+		const UPhysicsAsset* InPhysicsAsset, 
+		const FChaosClothSimulationModel& InClothSimulationModel, 
+		const TArray<TSharedRef<const FManagedArrayCollection>>& InCollections);
+
 	FChaosOutfitPiece(const UChaosClothAssetBase& ClothAssetBase, int32 ModelIndex);
 
 	FChaosOutfitPiece(const FChaosOutfitPiece& Other);
+
 	FChaosOutfitPiece(FChaosOutfitPiece&& Other);
 
 	FChaosOutfitPiece& operator=(const FChaosOutfitPiece& Other);
 	FChaosOutfitPiece& operator=(FChaosOutfitPiece&& Other);
 
-	// Custom serialize for TSharedRef
 	bool Serialize(FArchive& Ar);
 
 private:
@@ -75,6 +87,12 @@ public:
 	void Append(const UChaosOutfit& Other);
 	void Add(const UChaosClothAssetBase& ClothAssetBase);
 
+	void CopyTo(
+		TArray<FChaosOutfitPiece>& OutPieces,
+		FReferenceSkeleton& OutReferenceSkeleton,
+		TUniquePtr<FSkeletalMeshRenderData>& OutSkeletalMeshRenderData,
+		TArray<FSkeletalMaterial>& OutMaterials) const;
+
 	const TArray<FChaosOutfitPiece>& GetPieces() const
 	{
 		return Pieces;
@@ -89,9 +107,44 @@ public:
 	virtual void Serialize(FArchive& Ar) override;
 	//~ End UObject interface
 
+	struct FLODRenderData;
+	struct FRenderData;
+
+	template<typename RenderDataType UE_REQUIRES(std::is_same_v<RenderDataType, FSkeletalMeshRenderData> || std::is_same_v<RenderDataType, UChaosOutfit::FRenderData>)>
+	static void Init(
+		TArray<FChaosOutfitPiece>& OutPieces,
+		FReferenceSkeleton& OutReferenceSkeleton,
+		TUniquePtr<RenderDataType>& OutSkeletalMeshRenderData,
+		TArray<FSkeletalMaterial>& OutMaterials);
 private:
+	template<typename InLODRenderDataType, typename OutLODRenderDataType>
+	static void MergeLODRenderDatas(
+		const InLODRenderDataType& LODRenderData,
+		const TArray<FGuid>& AssetGuids,
+		const int32 MaterialOffset,
+		const FReferenceSkeleton& ReferenceSkeleton,
+		const TArray<int32>& BoneMap,
+		OutLODRenderDataType& OutLODRenderData);
+
+	template<typename InRenderDataType, typename OutRenderDataType>
+	static void Merge(
+		const FReferenceSkeleton& InReferenceSkeleton,
+		const InRenderDataType* const InRenderData,
+		const TArray<FSkeletalMaterial>& InMaterials,
+		const int32 NumPieces,
+		TFunction<FChaosOutfitPiece(int32)>&& MakePiece,
+		TArray<FChaosOutfitPiece>& OutPieces,
+		FReferenceSkeleton& OutReferenceSkeleton,
+		OutRenderDataType* const OutRenderData,
+		TArray<FSkeletalMaterial>& OutMaterials);
+
 	UPROPERTY()
 	TArray<FChaosOutfitPiece> Pieces;
 
+	UPROPERTY()
+	TArray<FSkeletalMaterial> Materials;
+
 	FReferenceSkeleton ReferenceSkeleton;
+
+	TUniquePtr<FRenderData> RenderData;
 };
diff --git a/Engine/Plugins/Experimental/ChaosOutfitAsset/Source/ChaosOutfitAssetEngine/Public/ChaosClothAsset/OutfitAsset.h b/Engine/Plugins/Experimental/ChaosOutfitAsset/Source/ChaosOutfitAssetEngine/Public/ChaosClothAsset/OutfitAsset.h
--- a/Engine/Plugins/Experimental/ChaosOutfitAsset/Source/ChaosOutfitAssetEngine/Public/ChaosClothAsset/OutfitAsset.h
+++ b/Engine/Plugins/Experimental/ChaosOutfitAsset/Source/ChaosOutfitAssetEngine/Public/ChaosClothAsset/OutfitAsset.h
@@ -21,7 +21,12 @@ public:
 	UChaosOutfitAsset(FVTableHelper& Helper);  // This is declared so we can use TUniquePtr<FClothSimulationModel> with just a forward declare of that class
 	virtual ~UChaosOutfitAsset() override;
 
-	void Build(const TObjectPtr<const UChaosOutfit> Outfit = nullptr);  // TODO: Remove default value once the terminal node code is in
+	void Build(const TObjectPtr<const UChaosOutfit> Outfit);
+
+	//~ Begin UObject interface
+	virtual void Serialize(FArchive& Ar) override;
+	virtual void PostLoad() override;
+	//~ End UObject interface
 
 private:
 	//~ Begin UChaosClothAssetBase interface
@@ -63,10 +68,11 @@ private:
 	{
 		return nullptr;
 	}
-#endif // WITH_EDITORONLY_DATA
-
+#endif
 	//~ End USkinnedAsset interface
 
+	void CalculateBounds();
+
 	UPROPERTY()
 	TArray<FChaosOutfitPiece> Pieces;
 };
