Chaos Outfit Asset - Removed render data serialization, and rebuild the entire outfit on loading until a proper DDC solution is worked out. This will fix crashes caused by the recent skeletal rendering serialization changes in CL 40789459. Will also need manual fixup on all outfit assets created before CL 40789459.

#rnx

diff --git a/Engine/Plugins/Experimental/ChaosOutfitAsset/Source/ChaosOutfitAssetEngine/Private/ChaosClothAsset/OutfitAsset.cpp b/Engine/Plugins/Experimental/ChaosOutfitAsset/Source/ChaosOutfitAssetEngine/Private/ChaosClothAsset/OutfitAsset.cpp
--- a/Engine/Plugins/Experimental/ChaosOutfitAsset/Source/ChaosOutfitAssetEngine/Private/ChaosClothAsset/OutfitAsset.cpp
+++ b/Engine/Plugins/Experimental/ChaosOutfitAsset/Source/ChaosOutfitAssetEngine/Private/ChaosClothAsset/OutfitAsset.cpp
@@ -103,15 +103,6 @@ void UChaosOutfitAsset::Serialize(FArchive& Ar)
 {
 	LLM_SCOPE_BYNAME(TEXT("Physics/Cloth"));
 	Super::Serialize(Ar);
-
-	{
-		LLM_SCOPE_BYNAME(TEXT("Physics/ClothRendering"));
-		if (Ar.IsLoading())
-		{
-			SetResourceForRendering(MakeUnique<FSkeletalMeshRenderData>());
-		}
-		GetResourceForRendering()->Serialize(Ar, this);
-	}
 }
 
 void UChaosOutfitAsset::PostLoad()
@@ -119,6 +110,9 @@ void UChaosOutfitAsset::PostLoad()
 	LLM_SCOPE_BYNAME(TEXT("Physics/Cloth"));
 	Super::PostLoad();
 
+	// Re-evaluate the Dataflow
+	GetDataflowInstance().UpdateOwnerAsset();  // TODO: Serialize the render data to DDC instead
+
 	if (FApp::CanEverRender())
 	{
 		InitResources();
