Chaos Cloth Outfit - Fixed the MakeOutfit ClothAssets property category name.

#rnx

diff --git a/Engine/Plugins/Experimental/ChaosOutfitAsset/Source/ChaosOutfitAssetDataflowNodes/Private/ChaosClothAsset/MakeOutfitNode.h b/Engine/Plugins/Experimental/ChaosOutfitAsset/Source/ChaosOutfitAssetDataflowNodes/Private/ChaosClothAsset/MakeOutfitNode.h
--- a/Engine/Plugins/Experimental/ChaosOutfitAsset/Source/ChaosOutfitAssetDataflowNodes/Private/ChaosClothAsset/MakeOutfitNode.h
+++ b/Engine/Plugins/Experimental/ChaosOutfitAsset/Source/ChaosOutfitAssetDataflowNodes/Private/ChaosClothAsset/MakeOutfitNode.h
@@ -47,6 +47,6 @@ private:
 	TObjectPtr<const UChaosOutfit> Outfit;
 
 	/** The cloth assets or outfits to import into a new outfit. */
-	UPROPERTY(EditAnywhere, Category = "ClothAssetToOutfit")
+	UPROPERTY(EditAnywhere, Category = "MakeOutfit")
 	TArray<FChaosClothAssetOrArrayAnyType> ClothAssets;
 };
