Chaos Outfit Asset - Added Outfit Collection.

#rnx

diff --git a/Engine/Plugins/Experimental/ChaosOutfitAsset/Source/ChaosOutfitAssetDataflowNodes/ChaosOutfitAssetDataflowNodes.Build.cs b/Engine/Plugins/Experimental/ChaosOutfitAsset/Source/ChaosOutfitAssetDataflowNodes/ChaosOutfitAssetDataflowNodes.Build.cs
--- a/Engine/Plugins/Experimental/ChaosOutfitAsset/Source/ChaosOutfitAssetDataflowNodes/ChaosOutfitAssetDataflowNodes.Build.cs
+++ b/Engine/Plugins/Experimental/ChaosOutfitAsset/Source/ChaosOutfitAssetDataflowNodes/ChaosOutfitAssetDataflowNodes.Build.cs
@@ -18,6 +18,7 @@ public class ChaosOutfitAssetDataflowNodes : ModuleRules
 				"ChaosOutfitAssetEngine",
 				"DataflowCore",
 				"DataflowNodes",
+				"Engine",
 			}
 		);
 	}
diff --git a/Engine/Plugins/Experimental/ChaosOutfitAsset/Source/ChaosOutfitAssetDataflowNodes/Private/ChaosClothAsset/ChaosOutfitAssetDataflowNodesModule.cpp b/Engine/Plugins/Experimental/ChaosOutfitAsset/Source/ChaosOutfitAssetDataflowNodes/Private/ChaosClothAsset/ChaosOutfitAssetDataflowNodesModule.cpp
--- a/Engine/Plugins/Experimental/ChaosOutfitAsset/Source/ChaosOutfitAssetDataflowNodes/Private/ChaosClothAsset/ChaosOutfitAssetDataflowNodesModule.cpp
+++ b/Engine/Plugins/Experimental/ChaosOutfitAsset/Source/ChaosOutfitAssetDataflowNodes/Private/ChaosClothAsset/ChaosOutfitAssetDataflowNodesModule.cpp
@@ -4,8 +4,10 @@
 #include "ChaosClothAsset/GetClothAssetNode.h"
 #include "ChaosClothAsset/GetOutfitAssetNode.h"
 #include "ChaosClothAsset/MakeOutfitNode.h"
+#include "ChaosClothAsset/MakeSizedOutfitNode.h"
 #include "ChaosClothAsset/MergeOutfitsNode.h"
 #include "ChaosClothAsset/OutfitAssetTerminalNode.h"
+#include "ChaosClothAsset/SizedOutfitSourceAnyType.h"
 #include "Dataflow/DataflowAnyTypeRegistry.h"
 #include "Dataflow/DataflowNode.h"
 #include "Dataflow/DataflowNodeColorsRegistry.h"
@@ -28,6 +30,7 @@ namespace UE::Chaos::ClothAsset
 			DATAFLOW_NODE_REGISTER_CREATION_FACTORY(FChaosGetClothAssetNode);
 			DATAFLOW_NODE_REGISTER_CREATION_FACTORY(FChaosGetOutfitAssetNode);
 			DATAFLOW_NODE_REGISTER_CREATION_FACTORY(FChaosOutfitAssetMakeOutfitNode);
+			DATAFLOW_NODE_REGISTER_CREATION_FACTORY(FChaosOutfitAssetMakeSizedOutfitNode);
 			DATAFLOW_NODE_REGISTER_CREATION_FACTORY(FChaosOutfitAssetMergeOutfitsNode);
 			DATAFLOW_NODE_REGISTER_CREATION_FACTORY(FChaosOutfitAssetTerminalNode);
 PRAGMA_DISABLE_DEPRECATION_WARNINGS
@@ -45,6 +48,8 @@ PRAGMA_ENABLE_DEPRECATION_WARNINGS
 			DATAFLOW_NODE_REGISTER_CREATION_FACTORY_PIN_SETTINGS_BY_TYPE("FChaosClothAssetArrayAnyType", OutfitAssetNodeHeaderColor, DefaultWireThickness);
 			UE_DATAFLOW_REGISTER_ANYTYPE(FChaosClothAssetOrArrayAnyType);
 			DATAFLOW_NODE_REGISTER_CREATION_FACTORY_PIN_SETTINGS_BY_TYPE("FChaosClothAssetOrArrayAnyType", OutfitAssetNodeHeaderColor, DefaultWireThickness);
+			UE_DATAFLOW_REGISTER_ANYTYPE(FChaosSizedOutfitSourceOrArrayAnyType);
+			DATAFLOW_NODE_REGISTER_CREATION_FACTORY_PIN_SETTINGS_BY_TYPE("FChaosSizedOutfitSourceOrArrayAnyType", OutfitAssetNodeHeaderColor, DefaultWireThickness);
 		}
 	}  // namespace Private
 
diff --git a/Engine/Plugins/Experimental/ChaosOutfitAsset/Source/ChaosOutfitAssetDataflowNodes/Private/ChaosClothAsset/MakeSizedOutfitNode.cpp b/Engine/Plugins/Experimental/ChaosOutfitAsset/Source/ChaosOutfitAssetDataflowNodes/Private/ChaosClothAsset/MakeSizedOutfitNode.cpp
new file mode 100644
--- /dev/null
+++ b/Engine/Plugins/Experimental/ChaosOutfitAsset/Source/ChaosOutfitAssetDataflowNodes/Private/ChaosClothAsset/MakeSizedOutfitNode.cpp
@@ -0,0 +1,202 @@
+// Copyright Epic Games, Inc. All Rights Reserved.
+
+#include "ChaosClothAsset/MakeSizedOutfitNode.h"
+#include "ChaosClothAsset/CollectionOutfitFacade.h"
+#include "ChaosClothAsset/ClothAssetAnyType.h"
+#include "ChaosClothAsset/Outfit.h"
+#include "Engine/SkeletalMesh.h"
+
+#include UE_INLINE_GENERATED_CPP_BY_NAME(MakeSizedOutfitNode)
+
+#define LOCTEXT_NAMESPACE "ChaosOutfitAssetMakeSizedOutfitNode"
+
+FChaosOutfitAssetMakeSizedOutfitNode::FChaosOutfitAssetMakeSizedOutfitNode(const UE::Dataflow::FNodeParameters& InParam, FGuid InGuid)
+	: FDataflowNode(InParam, InGuid)
+{
+	RegisterOutputConnection(&Outfit);
+	RegisterOutputConnection(&OutfitCollection);
+
+	for (int32 Index = 0; Index < NumInitialSizedOutfitSources; ++Index)
+	{
+		AddPins();
+	}
+	check(GetNumInputs() == NumRequiredInputs + NumInitialSizedOutfitSources);  // Update NumRequiredInputs when adding inputs (used by Serialize)
+}
+
+void FChaosOutfitAssetMakeSizedOutfitNode::Evaluate(UE::Dataflow::FContext& Context, const FDataflowOutput* Out) const
+{
+	if (Out->IsA(&Outfit) || Out->IsA(&OutfitCollection))
+	{
+		TObjectPtr<UChaosOutfit> OutOutfit = NewObject<UChaosOutfit>();
+		UE::Chaos::OutfitAsset::FCollectionOutfitFacade OutfitFacade(OutOutfit->GetOutfitCollection());
+
+		auto GetValidBodyPart = [](const FChaosSizedOutfitSource& SizedOutfitSource) -> const USkeletalMesh*
+			{
+				for (const TObjectPtr<const USkeletalMesh>& SourceBodyPart : SizedOutfitSource.SourceBodyParts)
+				{
+					if (SourceBodyPart)
+					{
+						return SourceBodyPart.Get();
+					}
+				}
+				return nullptr;
+			};
+
+		// Make a new GUID for this sized outfit asset
+		const FGuid OutfitGuid = FGuid::NewGuid();
+
+		for (int32 SourceIndex = 0; SourceIndex < SizedOutfitSources.Num(); ++SourceIndex)
+		{
+			const FChaosSizedOutfitSourceOrArrayType& SizedOutfitSourceOrArray = GetValue(Context, &SizedOutfitSources[SourceIndex]);
+
+			for (const FChaosSizedOutfitSource& SizedOutfitSource : SizedOutfitSourceOrArray.Array)
+			{
+				// Check size, only display a warning for empty sizes when not empty of cloth assets, otherwise skip
+				FString BodySizeName = SizedOutfitSource.SizeName;
+				if (BodySizeName.IsEmpty())
+				{
+					// Use the skeletal mesh name if the size name is empty
+					if (const USkeletalMesh* SkeletalMesh = GetValidBodyPart(SizedOutfitSource))
+					{
+						BodySizeName = SkeletalMesh->GetName();
+					}
+					else if (SizedOutfitSource.SourceAsset && SizedOutfitSource.SourceAsset->GetNumClothSimulationModels())
+					{
+						BodySizeName = TEXT("Default");
+					}
+					else
+					{
+						continue;  // Empty size, skip
+					}
+				}
+
+				// Add body size
+				if (OutfitFacade.HasBodySize(BodySizeName))
+				{
+					Context.Warning(FString::Printf(TEXT("The body size [%s] already exist and has now been replaced."), *BodySizeName), this, Out);
+				}
+
+				TArray<FString> BodyPartsSkeletalMeshes;
+				BodyPartsSkeletalMeshes.Reserve(SizedOutfitSource.SourceBodyParts.Num());
+				for (const TObjectPtr<const USkeletalMesh>& SourceBodyPart : SizedOutfitSource.SourceBodyParts)
+				{
+					if (SourceBodyPart)
+					{
+						BodyPartsSkeletalMeshes.Emplace(SourceBodyPart->GetPathName());
+					}
+				}
+				TMap<FString, float> Measurements;  // TODO: Read measurements
+				const int32 BodySize = OutfitFacade.AddBodySize(BodySizeName, BodyPartsSkeletalMeshes, Measurements);
+
+				if (const TObjectPtr<const UChaosClothAssetBase> ClothAsset = SizedOutfitSource.SourceAsset)
+				{
+					const int32 NumPieces = ClothAsset->GetNumClothSimulationModels();
+					if (!NumPieces)
+					{
+						Context.Warning(FString::Printf(TEXT("Empty source asset [%s] specified for size [%s]."), *ClothAsset->GetName(), *BodySizeName), this, Out);
+						continue;
+					}
+					TMap<FGuid, FString> Pieces;
+					Pieces.Reserve(NumPieces);
+					for (int32 Index = 0; Index < NumPieces; ++Index)
+					{
+						Pieces.Emplace(
+							ClothAsset->GetAssetGuid(Index),
+							ClothAsset->GetClothSimulationModelName(Index).ToString());
+					}
+
+					// Add size metadata
+					OutfitFacade.AddOutfit(
+						OutfitGuid,
+						BodySize,
+						ClothAsset->GetName(),
+						Pieces);
+
+					// Add asset
+					OutOutfit->Add(*ClothAsset);
+				}
+			}
+		}
+
+		SetValue<TObjectPtr<const UChaosOutfit>>(Context, OutOutfit, &Outfit);
+		SetValue(Context, OutOutfit->GetOutfitCollection(), &OutfitCollection);
+	}
+}
+
+TArray<UE::Dataflow::FPin> FChaosOutfitAssetMakeSizedOutfitNode::AddPins()
+{
+	const int32 Index = SizedOutfitSources.AddDefaulted();
+	const FDataflowInput& Input = RegisterInputArrayConnection(GetConnectionReference(Index));
+	return { { UE::Dataflow::FPin::EDirection::INPUT, Input.GetType(), Input.GetName() } };
+}
+
+TArray<UE::Dataflow::FPin> FChaosOutfitAssetMakeSizedOutfitNode::GetPinsToRemove() const
+{
+	const int32 Index = SizedOutfitSources.Num() - 1;
+	check(SizedOutfitSources.IsValidIndex(Index));
+	if (const FDataflowInput* const Input = FindInput(GetConnectionReference(Index)))
+	{
+		return { { UE::Dataflow::FPin::EDirection::INPUT, Input->GetType(), Input->GetName() } };
+	}
+	return Super::GetPinsToRemove();
+}
+
+void FChaosOutfitAssetMakeSizedOutfitNode::OnPinRemoved(const UE::Dataflow::FPin& Pin)
+{
+	const int32 Index = SizedOutfitSources.Num() - 1;
+	check(SizedOutfitSources.IsValidIndex(Index));
+#if DO_CHECK
+	const FDataflowInput* const Input = FindInput(GetConnectionReference(Index));
+	check(Input);
+	check(Input->GetName() == Pin.Name);
+	check(Input->GetType() == Pin.Type);
+#endif
+	SizedOutfitSources.SetNum(Index);
+
+	return Super::OnPinRemoved(Pin);
+}
+
+void FChaosOutfitAssetMakeSizedOutfitNode::PostSerialize(const FArchive& Ar)
+{
+	// Added pins need to be restored when loading to make sure they get reconnected
+	if (Ar.IsLoading())
+	{
+		for (int32 Index = 0; Index < NumInitialSizedOutfitSources; ++Index)
+		{
+			check(FindInput(GetConnectionReference(Index)));
+		}
+
+		for (int32 Index = NumInitialSizedOutfitSources; Index < SizedOutfitSources.Num(); ++Index)
+		{
+			FindOrRegisterInputArrayConnection(GetConnectionReference(Index));
+		}
+		if (Ar.IsTransacting())
+		{
+			const int32 OrigNumRegisteredInputs = GetNumInputs();
+			check(OrigNumRegisteredInputs >= NumRequiredInputs + NumInitialSizedOutfitSources);
+			const int32 OrigNumSizedOutfitSources = SizedOutfitSources.Num();
+			const int32 OrigNumRegisteredSizedOutfitSources = (OrigNumRegisteredInputs - NumRequiredInputs);
+			if (OrigNumRegisteredSizedOutfitSources > OrigNumSizedOutfitSources)
+			{
+				// Inputs have been removed, temporarily expand SizedOutfitInputs so we can get connection references
+				SizedOutfitSources.SetNum(OrigNumRegisteredSizedOutfitSources);
+				for (int32 Index = OrigNumSizedOutfitSources; Index < SizedOutfitSources.Num(); ++Index)
+				{
+					UnregisterInputConnection(GetConnectionReference(Index));
+				}
+				SizedOutfitSources.SetNum(OrigNumSizedOutfitSources);
+			}
+		}
+		else
+		{
+			ensureAlways(SizedOutfitSources.Num() + NumRequiredInputs == GetNumInputs());
+		}
+	}
+}
+
+UE::Dataflow::TConnectionReference<FChaosSizedOutfitSourceOrArrayAnyType> FChaosOutfitAssetMakeSizedOutfitNode::GetConnectionReference(int32 Index) const
+{
+	return { &SizedOutfitSources[Index], Index, &SizedOutfitSources };
+}
+
+#undef LOCTEXT_NAMESPACE
diff --git a/Engine/Plugins/Experimental/ChaosOutfitAsset/Source/ChaosOutfitAssetDataflowNodes/Private/ChaosClothAsset/MakeSizedOutfitNode.h b/Engine/Plugins/Experimental/ChaosOutfitAsset/Source/ChaosOutfitAssetDataflowNodes/Private/ChaosClothAsset/MakeSizedOutfitNode.h
new file mode 100644
--- /dev/null
+++ b/Engine/Plugins/Experimental/ChaosOutfitAsset/Source/ChaosOutfitAssetDataflowNodes/Private/ChaosClothAsset/MakeSizedOutfitNode.h
@@ -0,0 +1,55 @@
+// Copyright Epic Games, Inc. All Rights Reserved.
+
+#pragma once 
+
+#include "ChaosClothAsset/SizedOutfitSourceAnyType.h"
+#include "Dataflow/DataflowNode.h"
+#include "GeometryCollection/ManagedArrayCollection.h"
+#include "MakeSizedOutfitNode.generated.h"
+
+class UChaosOutfit;
+
+/** Add multiple cloth asset objects to an outfit collection. */
+USTRUCT(Meta = (Experimental, DataflowOutfit))
+struct FChaosOutfitAssetMakeSizedOutfitNode final : public FDataflowNode
+{
+	GENERATED_USTRUCT_BODY()
+	DATAFLOW_NODE_DEFINE_INTERNAL(FChaosOutfitAssetMakeSizedOutfitNode, "MakeSizedOutfit", "Outfit", "Cloth Asset Make Sized Outfit")
+
+public:
+	FChaosOutfitAssetMakeSizedOutfitNode(const UE::Dataflow::FNodeParameters& InParam, FGuid InGuid = FGuid::NewGuid());
+
+private:
+	//~ Begin FDataflowNode interface
+	virtual void Evaluate(UE::Dataflow::FContext& Context, const FDataflowOutput* Out) const override;
+	virtual TArray<UE::Dataflow::FPin> AddPins() override;
+	virtual bool CanAddPin() const override
+	{
+		return true;
+	}
+	virtual bool CanRemovePin() const override
+	{
+		return SizedOutfitSources.Num() > NumInitialSizedOutfitSources;
+	}
+	virtual TArray<UE::Dataflow::FPin> GetPinsToRemove() const override;
+	virtual void OnPinRemoved(const UE::Dataflow::FPin& Pin) override;
+	virtual void PostSerialize(const FArchive& Ar) override;
+	//~ End FDataflowNode interface
+
+	UE::Dataflow::TConnectionReference<FChaosSizedOutfitSourceOrArrayAnyType> GetConnectionReference(int32 Index) const;
+
+	static constexpr int32 NumRequiredInputs = 0;
+	static constexpr int32 NumInitialSizedOutfitSources = 1;
+
+	/** The outfit output. */
+	UPROPERTY(Meta = (DataflowOutput))
+	TObjectPtr<const UChaosOutfit> Outfit;
+
+	/** The outfit collection output, provided for convenience as a view into the outfit object metadata. */
+	UPROPERTY(Meta = (DataflowOutput))
+	FManagedArrayCollection OutfitCollection;
+
+	/** The cloth assets/outfits and bidy references to import into as a new sized outfit. */
+	UPROPERTY(EditAnywhere, Category = "MakeSizedOutfit")
+	TArray<FChaosSizedOutfitSourceOrArrayAnyType> SizedOutfitSources;
+};
diff --git a/Engine/Plugins/Experimental/ChaosOutfitAsset/Source/ChaosOutfitAssetDataflowNodes/Private/ChaosClothAsset/SizedOutfitSource.h b/Engine/Plugins/Experimental/ChaosOutfitAsset/Source/ChaosOutfitAssetDataflowNodes/Private/ChaosClothAsset/SizedOutfitSource.h
new file mode 100644
--- /dev/null
+++ b/Engine/Plugins/Experimental/ChaosOutfitAsset/Source/ChaosOutfitAssetDataflowNodes/Private/ChaosClothAsset/SizedOutfitSource.h
@@ -0,0 +1,42 @@
+// Copyright Epic Games, Inc. All Rights Reserved.
+
+#pragma once 
+
+#include "Containers/Array.h"
+#include "UObject/ObjectPtr.h"
+#include "SizedOutfitSource.generated.h"
+
+class UChaosClothAssetBase;
+class USkeletalMesh;
+
+/**
+ * Input structure for setting up a single sized outfit.
+ */
+USTRUCT(BlueprintType)  // BlueprintType to allow this structure as a variable type.
+struct FChaosSizedOutfitSource final
+{
+	GENERATED_USTRUCT_BODY()
+
+	/**
+	 * The cloth or outfit asset to assign for this body size.
+	 * The asset must be the exact same garment representation for each specified sizes.
+	 * If multiple garments in multiple sizes are needed, they will have to be composed in another Outfit asset.
+	 */
+	UPROPERTY(EditAnywhere, Category = "OutfitSourceAsset")
+	TObjectPtr<const UChaosClothAssetBase> SourceAsset;
+
+	/**
+	 * The unique name of this body size.
+	 * The name of the first valid body part skeletal mesh will be used if empty.
+	 */
+	UPROPERTY(EditAnywhere, Category = "OutfitBodySize")
+	FString SizeName;
+
+	/**
+	 * The list of body part skeletal meshes making up the source body for this size name.
+	 * Usually one single MetaHuman merged body+head skeletal mesh.
+	 */
+	UPROPERTY(EditAnywhere, Category = "OutfitBodySize", Meta = (NoElementDuplicate))
+	TArray<TObjectPtr<const USkeletalMesh>> SourceBodyParts = { nullptr };
+
+};
diff --git a/Engine/Plugins/Experimental/ChaosOutfitAsset/Source/ChaosOutfitAssetDataflowNodes/Private/ChaosClothAsset/SizedOutfitSourceAnyType.h b/Engine/Plugins/Experimental/ChaosOutfitAsset/Source/ChaosOutfitAssetDataflowNodes/Private/ChaosClothAsset/SizedOutfitSourceAnyType.h
new file mode 100644
--- /dev/null
+++ b/Engine/Plugins/Experimental/ChaosOutfitAsset/Source/ChaosOutfitAssetDataflowNodes/Private/ChaosClothAsset/SizedOutfitSourceAnyType.h
@@ -0,0 +1,51 @@
+// Copyright Epic Games, Inc. All Rights Reserved.
+
+#pragma once 
+
+#include "Dataflow/DataflowAnyType.h"
+#include "Dataflow/DataflowTypePolicy.h"
+#include "ChaosClothAsset/SizedOutfitSource.h"
+#include "Misc/TVariant.h"
+
+#include "SizedOutfitSourceAnyType.generated.h"
+
+USTRUCT()
+struct FChaosSizedOutfitSourceOrArrayType
+{
+	GENERATED_BODY()
+
+	UPROPERTY()
+	TArray<FChaosSizedOutfitSource> Array;
+
+	FChaosSizedOutfitSourceOrArrayType() = default;
+
+	FChaosSizedOutfitSourceOrArrayType(const FChaosSizedOutfitSource& SizedOutfitSource)
+		: Array(&SizedOutfitSource, 1)
+	{}
+
+	FChaosSizedOutfitSourceOrArrayType(const TArray<FChaosSizedOutfitSource>& SizedOutfitSources)
+		: Array(SizedOutfitSources)
+	{}
+};
+
+UE_DATAFLOW_POLICY_DECLARE_TYPENAME(FChaosSizedOutfitSource)
+UE_DATAFLOW_POLICY_DECLARE_TYPENAME(FChaosSizedOutfitSourceOrArrayType)
+
+/** FChaosSizedOutfitSource or array types. */
+USTRUCT()
+struct FChaosSizedOutfitSourceOrArrayAnyType : public FDataflowAnyType
+{
+	GENERATED_BODY()
+
+	struct FPolicy :
+		public TDataflowMultiTypePolicy<
+			FChaosSizedOutfitSource,
+			TArray<FChaosSizedOutfitSource>>
+	{};
+
+	using FPolicyType = FPolicy;
+	using FStorageType = FChaosSizedOutfitSourceOrArrayType;
+
+	UPROPERTY(EditAnywhere, Category = Value)
+	FChaosSizedOutfitSourceOrArrayType Value;
+};
diff --git a/Engine/Plugins/Experimental/ChaosOutfitAsset/Source/ChaosOutfitAssetEngine/Private/ChaosClothAsset/CollectionOutfitFacade.cpp b/Engine/Plugins/Experimental/ChaosOutfitAsset/Source/ChaosOutfitAssetEngine/Private/ChaosClothAsset/CollectionOutfitFacade.cpp
new file mode 100644
--- /dev/null
+++ b/Engine/Plugins/Experimental/ChaosOutfitAsset/Source/ChaosOutfitAssetEngine/Private/ChaosClothAsset/CollectionOutfitFacade.cpp
@@ -0,0 +1,490 @@
+// Copyright Epic Games, Inc. All Rights Reserved.
+
+#include "ChaosClothAsset/CollectionOutfitFacade.h"
+#include "ChaosClothAsset/OutfitCollection.h"
+#include "GeometryCollection/ManagedArrayCollection.h"
+
+#define LOCTEXT_NAMESPACE "CollectionOutfitFacade"
+
+namespace UE::Chaos::OutfitAsset
+{
+	namespace Private
+	{
+		// Utility class to transfer const access from the pointer to the pointed data
+		template<typename T>
+		class TConstPreservingPointer
+		{
+		public:
+			TConstPreservingPointer(T* const InPtr = nullptr) : Ptr(InPtr) {}
+			TConstPreservingPointer(const TConstPreservingPointer&) = delete;  // Non copyable as to preserve the source const access
+			TConstPreservingPointer& operator=(const TConstPreservingPointer&) = delete;
+			TConstPreservingPointer& operator=(T* const InPtr)
+			{
+				Ptr = InPtr;
+				return *this;
+			}
+			bool IsValid() const
+			{
+				return Ptr != nullptr;
+			}
+			operator bool() const
+			{
+				return IsValid();
+			}
+			const T* operator->() const
+			{
+				return Ptr;
+			}
+			const T& operator*() const
+			{
+				return *Ptr;
+			}
+			T* operator->()
+			{
+				return Ptr;
+			}
+			T& operator*()
+			{
+				check(IsValid());
+				return *Ptr;
+			}
+		private:
+			T* Ptr = nullptr;
+		};
+
+		template<typename T>
+		using TManagedArrayPtr = TConstPreservingPointer<TManagedArray<T>>;
+
+		// Outfit collection schema implementation
+		struct FOutfitCollection
+		{
+			FManagedArrayCollection& ManagedArrayCollection;
+
+			struct
+			{
+				TManagedArrayPtr<FGuid> Guid;  // Outfit guid to distinguish between separated sized outfit merged into a single outfit
+				TManagedArrayPtr<int32> BodySize;  // Body size index in the BodySizes table
+				TManagedArrayPtr<FString> Name;  // Name of the source asset making this outfit size
+				TManagedArrayPtr<int32> PiecesStart;  // Outfit piece start index in the Pieces table
+				TManagedArrayPtr<int32> PiecesCount;  // Number of outfit pieces for this outfit
+			} Outfits;
+
+			struct
+			{
+				TManagedArrayPtr<FGuid> Guid;  // Outfit Piece GUID (same as FChaosOutfitPiece::Guid), in case the name is duplicated
+				TManagedArrayPtr<FString> Name;  // Outfit Piece Name (for debugging, same as FChaosOutfitPiece::Name)
+			} Pieces;
+
+			struct
+			{
+				TManagedArrayPtr<FString> Name;  // Name of this body size
+				TManagedArrayPtr<int32> BodyPartsStart;  // Body part start index in the BodyParts table
+				TManagedArrayPtr<int32> BodyPartsCount;  // Number of body parts for this body size
+			} BodySizes;
+
+			struct
+			{
+				TManagedArrayPtr<FString> SkeletalMesh;  // Body parts skeletal mesh (several per sizes, e.g. body, head)
+			} BodyParts;
+
+			struct
+			{
+				TManagedArrayPtr<FString> Name;  // Measurements name (several per sizes, e.g. one for Height, Waist, ...etc.)
+			} Measurements;
+
+			explicit FOutfitCollection(FManagedArrayCollection& InManagedArrayCollection)
+				: ManagedArrayCollection(InManagedArrayCollection)
+			{
+				using namespace OutfitCollection;
+				Outfits.Guid = ManagedArrayCollection.FindAttribute<FGuid>(Attribute::Outfits::Guid, Group::Outfits);
+				Outfits.BodySize = ManagedArrayCollection.FindAttribute<int32>(Attribute::Outfits::BodySize, Group::Outfits);
+				Outfits.Name = ManagedArrayCollection.FindAttribute<FString>(Attribute::Outfits::Name, Group::Outfits);
+				Outfits.PiecesStart = ManagedArrayCollection.FindAttribute<int32>(Attribute::Outfits::PiecesStart, Group::Outfits);
+				Outfits.PiecesCount = ManagedArrayCollection.FindAttribute<int32>(Attribute::Outfits::PiecesCount, Group::Outfits);
+				Pieces.Guid = ManagedArrayCollection.FindAttribute<FGuid>(Attribute::Pieces::Guid, Group::Pieces);
+				Pieces.Name = ManagedArrayCollection.FindAttribute<FString>(Attribute::Pieces::Name, Group::Pieces);
+				BodySizes.Name = ManagedArrayCollection.FindAttribute<FString>(Attribute::BodySizes::Name, Group::BodySizes);
+				BodySizes.BodyPartsStart = ManagedArrayCollection.FindAttribute<int32>(Attribute::BodySizes::BodyPartsStart, Group::BodySizes);
+				BodySizes.BodyPartsCount = ManagedArrayCollection.FindAttribute<int32>(Attribute::BodySizes::BodyPartsCount, Group::BodySizes);
+				BodyParts.SkeletalMesh = ManagedArrayCollection.FindAttribute<FString>(Attribute::BodyParts::SkeletalMesh, Group::BodyParts);
+				Measurements.Name = ManagedArrayCollection.FindAttribute<FString>(Attribute::Measurements::Name, Group::Measurements);
+			}
+
+			bool IsValid() const
+			{
+				using namespace OutfitCollection;
+				return
+					Outfits.Guid &&
+					Outfits.BodySize &&
+					Outfits.Name &&
+					Outfits.PiecesStart &&
+					Outfits.PiecesCount &&
+					Pieces.Guid &&
+					Pieces.Name &&
+					BodySizes.Name &&
+					BodySizes.BodyPartsStart &&
+					BodySizes.BodyPartsCount &&
+					BodyParts.SkeletalMesh &&
+					Measurements.Name;
+			}
+
+			void DefineSchema()
+			{
+				using namespace OutfitCollection;
+				Outfits.Guid = &ManagedArrayCollection.AddAttribute<FGuid>(Attribute::Outfits::Guid, Group::Outfits);
+				Outfits.BodySize = &ManagedArrayCollection.AddAttribute<int32>(Attribute::Outfits::BodySize, Group::Outfits, Group::BodySizes);
+				Outfits.Name = &ManagedArrayCollection.AddAttribute<FString>(Attribute::Outfits::Name, Group::Outfits);
+				Outfits.PiecesStart = &ManagedArrayCollection.AddAttribute<int32>(Attribute::Outfits::PiecesStart, Group::Outfits, Group::Pieces);
+				Outfits.PiecesCount = &ManagedArrayCollection.AddAttribute<int32>(Attribute::Outfits::PiecesCount, Group::Outfits);
+				Pieces.Guid = &ManagedArrayCollection.AddAttribute<FGuid>(Attribute::Pieces::Guid, Group::Pieces);
+				Pieces.Name = &ManagedArrayCollection.AddAttribute<FString>(Attribute::Pieces::Name, Group::Pieces);
+				BodySizes.Name = &ManagedArrayCollection.AddAttribute<FString>(Attribute::BodySizes::Name, Group::BodySizes);
+				BodySizes.BodyPartsStart = &ManagedArrayCollection.AddAttribute<int32>(Attribute::BodySizes::BodyPartsStart, Group::BodySizes, Group::BodyParts);
+				BodySizes.BodyPartsCount = &ManagedArrayCollection.AddAttribute<int32>(Attribute::BodySizes::BodyPartsCount, Group::BodySizes);
+				BodyParts.SkeletalMesh = &ManagedArrayCollection.AddAttribute<FString>(Attribute::BodyParts::SkeletalMesh, Group::BodyParts);
+				Measurements.Name = &ManagedArrayCollection.AddAttribute<FString>(Attribute::Measurements::Name, Group::Measurements);
+			}
+		};
+	}  // namespace Private
+
+	FCollectionOutfitConstFacade::FCollectionOutfitConstFacade(const TSharedRef<const FManagedArrayCollection>& InManagedArrayCollection)
+		: OutfitCollection(MakeUnique<const Private::FOutfitCollection>(const_cast<FManagedArrayCollection&>(InManagedArrayCollection.Get())))
+		, ManagedArrayCollection(InManagedArrayCollection.ToSharedPtr())
+	{}
+
+	FCollectionOutfitConstFacade::FCollectionOutfitConstFacade(const FManagedArrayCollection& InManagedArrayCollection)
+		: OutfitCollection(MakeUnique<const Private::FOutfitCollection>(const_cast<FManagedArrayCollection&>(InManagedArrayCollection)))
+	{}
+
+	FCollectionOutfitConstFacade::~FCollectionOutfitConstFacade() = default;
+
+	bool FCollectionOutfitConstFacade::IsValid() const
+	{
+		return GetOutfitCollection()->IsValid();
+	}
+
+	TArray<FGuid> FCollectionOutfitConstFacade::GetOutfitGuids() const
+	{
+		check(IsValid());
+		TArray<FGuid> OutfitGuids;
+		OutfitGuids.Reserve(GetOutfitCollection()->Outfits.Guid->Num());
+		for (const FGuid& Guid : *GetOutfitCollection()->Outfits.Guid)
+		{
+			OutfitGuids.AddUnique(Guid);
+		}
+		return OutfitGuids;
+	}
+
+	TArray<FString> FCollectionOutfitConstFacade::GetOutfitBodyPartsSkeletalMeshes() const
+	{
+		check(IsValid());
+		TArray<FString> BodyParts;
+		BodyParts.Reserve(GetOutfitCollection()->BodyParts.SkeletalMesh->Num());
+		for (const FString& SkeletalMesh : (*GetOutfitCollection()->BodyParts.SkeletalMesh))
+		{
+			BodyParts.AddUnique(SkeletalMesh);
+		}
+		return BodyParts;
+	}
+
+	TArray<int32> FCollectionOutfitConstFacade::GetOutfitBodySizes(const FGuid& Guid) const
+	{
+		check(IsValid());
+		const int32 NumOutfits = GetOutfitCollection()->Outfits.Guid->Num();
+		TArray<int32> OutfitBodySizes;
+		OutfitBodySizes.Reserve(NumOutfits);
+		for (int32 OutfitIndex = 0; OutfitIndex < NumOutfits; ++OutfitIndex)
+		{
+			if ((*GetOutfitCollection()->Outfits.Guid)[OutfitIndex] == Guid)
+			{
+				OutfitBodySizes.Add((*GetOutfitCollection()->Outfits.BodySize)[OutfitIndex]);
+			}
+		}
+		return OutfitBodySizes;
+	}
+
+	const FString& FCollectionOutfitConstFacade::GetOutfitName(const FGuid& Guid, const int32 BodySize) const
+	{
+		for (int32 OutfitIndex = 0; OutfitIndex < GetOutfitCollection()->Outfits.Guid->Num(); ++OutfitIndex)
+		{
+			if ((*GetOutfitCollection()->Outfits.Guid)[OutfitIndex] == Guid &&
+				(*GetOutfitCollection()->Outfits.BodySize)[OutfitIndex] == BodySize)
+			{
+				return (*GetOutfitCollection()->Outfits.Name)[OutfitIndex];
+			}
+		}
+		static const FString EmptyString;
+		return EmptyString;
+	}
+
+	TMap<FGuid, FString> FCollectionOutfitConstFacade::GetOutfitPieces(const FGuid& Guid, const int32 BodySize) const
+	{
+		check(IsValid());
+
+		auto GetOutfitPiecesRange = [this](const FGuid& Guid, const int32 BodySize) -> FIntVector2
+			{
+				for (int32 OutfitIndex = 0; OutfitIndex < GetOutfitCollection()->Outfits.Guid->Num(); ++OutfitIndex)
+				{
+					if ((*GetOutfitCollection()->Outfits.Guid)[OutfitIndex] == Guid &&
+						(*GetOutfitCollection()->Outfits.BodySize)[OutfitIndex] == BodySize)
+					{
+						return FIntVector2(
+							(*GetOutfitCollection()->Outfits.PiecesStart)[OutfitIndex],
+							(*GetOutfitCollection()->Outfits.PiecesCount)[OutfitIndex]);
+					}
+				}
+				return FIntVector2(INDEX_NONE, 0);
+			};
+
+		const FIntVector2 Range = GetOutfitPiecesRange(Guid, BodySize);
+		const int32 PiecesStart = Range[0];
+		const int32 PiecesCount = Range[1];
+		const TConstArrayView<FGuid> Guids =
+			PiecesCount &&
+			GetOutfitCollection()->Pieces.Guid->IsValidIndex(PiecesStart) &&
+			GetOutfitCollection()->Pieces.Guid->IsValidIndex(PiecesStart + PiecesCount - 1) ?
+			TConstArrayView<FGuid>(GetOutfitCollection()->Pieces.Guid->GetData() + PiecesStart, PiecesCount) :
+			TConstArrayView<FGuid>();
+		const TConstArrayView<FString> Names =
+			PiecesCount &&
+			GetOutfitCollection()->Pieces.Name->IsValidIndex(PiecesStart) &&
+			GetOutfitCollection()->Pieces.Name->IsValidIndex(PiecesStart + PiecesCount - 1) ?
+			TConstArrayView<FString>(GetOutfitCollection()->Pieces.Name->GetData() + PiecesStart, PiecesCount) :
+			TConstArrayView<FString>();
+		check(Guids.Num() == Names.Num());
+
+		TMap<FGuid, FString> OutfitPieces;
+		OutfitPieces.Reserve(Guids.Num());
+		for (int32 Index = 0; Index < Guids.Num(); ++Index)
+		{
+			OutfitPieces.Emplace(Guids[Index], Names[Index]);
+		}
+		return OutfitPieces;
+	}
+
+	int32 FCollectionOutfitConstFacade::GetNumBodySizes() const
+	{
+		check(IsValid());
+		return GetOutfitCollection()->BodySizes.Name->Num();
+	}
+
+	bool FCollectionOutfitConstFacade::HasBodySize(const FString& Name) const
+	{
+		check(IsValid());
+		return GetOutfitCollection()->BodySizes.Name->Find(Name) != INDEX_NONE;
+	}
+
+	const FString& FCollectionOutfitConstFacade::GetBodySizeName(int32 BodySize) const
+	{
+		check(IsValid());
+		static const FString EmptyString;
+		return GetOutfitCollection()->BodySizes.Name->IsValidIndex(BodySize) ?
+			(*GetOutfitCollection()->BodySizes.Name)[BodySize] : EmptyString;
+	}
+
+	TConstArrayView<FString> FCollectionOutfitConstFacade::GetBodySizeBodyPartsSkeletalMeshes(int32 BodySize) const
+	{
+		check(IsValid());
+		if (GetOutfitCollection()->BodySizes.Name->IsValidIndex(BodySize))
+		{
+			const int32 BodyPartsStart = (*GetOutfitCollection()->BodySizes.BodyPartsStart)[BodySize];
+			const int32 BodyPartsCount = (*GetOutfitCollection()->BodySizes.BodyPartsCount)[BodySize];
+			if (BodyPartsCount &&
+				GetOutfitCollection()->BodyParts.SkeletalMesh->IsValidIndex(BodyPartsStart) &&
+				GetOutfitCollection()->BodyParts.SkeletalMesh->IsValidIndex(BodyPartsStart + BodyPartsCount - 1))
+			{
+				return TConstArrayView<FString>(GetOutfitCollection()->BodyParts.SkeletalMesh->GetData() + BodyPartsStart, BodyPartsCount);
+			}
+		}
+		return TConstArrayView<FString>();
+	}
+
+	TMap<FString, float> FCollectionOutfitConstFacade::GetBodySizeMeasurements(const int32 BodySize) const
+	{
+		using namespace OutfitCollection;
+		check(IsValid());
+		TMap<FString, float> Measurements;
+		if (GetOutfitCollection()->BodySizes.Name->IsValidIndex(BodySize))
+		{
+			const FName Name((*GetOutfitCollection()->BodySizes.Name)[BodySize]);
+			if (const TManagedArray<float>* Values = GetOutfitCollection()->ManagedArrayCollection.FindAttributeTyped<float>(Name, Group::Measurements))
+			{
+				const int32 NumMeasurements = GetOutfitCollection()->Measurements.Name->Num();
+				Measurements.Reserve(NumMeasurements);
+				for (int32 Index = 0; Index < NumMeasurements; ++Index)
+				{
+					Measurements.Emplace((*GetOutfitCollection()->Measurements.Name)[Index], (*Values)[Index]);
+				}
+			}
+		}
+		return Measurements;
+	}
+
+	FCollectionOutfitFacade::FCollectionOutfitFacade(const TSharedRef<FManagedArrayCollection>& ManagedArrayCollection) :
+		FCollectionOutfitConstFacade(ManagedArrayCollection)
+	{
+	}
+
+	FCollectionOutfitFacade::FCollectionOutfitFacade(FManagedArrayCollection& InManagedArrayCollection)
+		: FCollectionOutfitConstFacade(InManagedArrayCollection)
+	{
+	}
+
+	FCollectionOutfitFacade::~FCollectionOutfitFacade() = default;
+
+	void FCollectionOutfitFacade::DefineSchema()
+	{
+		GetOutfitCollection()->DefineSchema();
+		check(IsValid());
+	}
+
+	int32 FCollectionOutfitFacade::AddBodySize(const FString& Name, const TConstArrayView<FString> BodyPartsSkeletalMeshes, const TMap<FString, float>& Measurements)
+	{
+		using namespace OutfitCollection;
+		check(IsValid());
+
+		int32 BodySize = GetOutfitCollection()->BodySizes.Name->Find(Name);
+		if (BodySize == INDEX_NONE)
+		{
+			BodySize = GetOutfitCollection()->ManagedArrayCollection.AddElements(1, OutfitCollection::Group::BodySizes);
+			(*GetOutfitCollection()->BodySizes.Name)[BodySize] = Name;
+		}
+		// Add or replace body parts
+		int32& BodyPartsCount = (*GetOutfitCollection()->BodySizes.BodyPartsCount)[BodySize];
+		int32& BodyPartsStart = (*GetOutfitCollection()->BodySizes.BodyPartsStart)[BodySize];
+		if (BodyPartsCount && BodyPartsStart != INDEX_NONE)
+		{
+			GetOutfitCollection()->ManagedArrayCollection.RemoveElements(OutfitCollection::Group::BodyParts, BodyPartsCount, BodyPartsStart);
+		}
+
+		BodyPartsCount = BodyPartsSkeletalMeshes.Num();
+		if (BodyPartsCount)
+		{
+			BodyPartsStart = GetOutfitCollection()->ManagedArrayCollection.AddElements(BodyPartsCount, Group::BodyParts);
+			for (int32 Index = 0; Index < BodyPartsCount; ++Index)
+			{
+				(*GetOutfitCollection()->BodyParts.SkeletalMesh)[BodyPartsStart + Index] = BodyPartsSkeletalMeshes[Index];
+			}
+		}
+		else
+		{
+			check(BodyPartsStart == INDEX_NONE);
+		}
+
+		// Add or replace body size measurements
+		const FName AttributeName(Name);
+		TManagedArray<float>* Values = GetOutfitCollection()->ManagedArrayCollection.FindAttributeTyped<float>(AttributeName, Group::Measurements);
+		if (!Values)
+		{
+			Values = &GetOutfitCollection()->ManagedArrayCollection.AddAttribute<float>(AttributeName, Group::Measurements);
+		}
+		else
+		{
+			// Clear all existing measurements
+			Values->Fill(0.f);
+		}
+		for (const TPair<FString, float>& Measurement : Measurements)
+		{
+			int32 Index = GetOutfitCollection()->Measurements.Name->Find(Measurement.Key);
+			if (Index == INDEX_NONE)
+			{
+				// Add missing key
+				Index = GetOutfitCollection()->ManagedArrayCollection.AddElements(1, OutfitCollection::Group::Measurements);
+				(*GetOutfitCollection()->Measurements.Name)[Index] = Measurement.Key;
+			}
+			(*Values)[Index] = Measurement.Value;
+		}
+
+		return BodySize;
+	}
+
+	void FCollectionOutfitFacade::AddOutfit(const FGuid& Guid, const int32 BodySize, const FString& Name, const TMap<FGuid, FString>& Pieces)
+	{
+		using namespace OutfitCollection;
+		check(IsValid());
+
+		auto FindOutfitIndex = [this](const FGuid& Guid, const int32 BodySize) -> int32
+			{
+				for (int32 Index = 0; Index < GetOutfitCollection()->Outfits.Guid->Num(); ++Index)
+				{
+					if ((*GetOutfitCollection()->Outfits.Guid)[Index] == Guid &&
+						(*GetOutfitCollection()->Outfits.BodySize)[Index] == BodySize)
+					{
+						return Index;
+					}
+				}
+				return INDEX_NONE;
+			};
+
+		int32 OutfitIndex = FindOutfitIndex(Guid, BodySize);
+		if (OutfitIndex == INDEX_NONE)
+		{
+			OutfitIndex = GetOutfitCollection()->ManagedArrayCollection.AddElements(1, OutfitCollection::Group::Outfits);
+			(*GetOutfitCollection()->Outfits.Guid)[OutfitIndex] = Guid;
+			(*GetOutfitCollection()->Outfits.BodySize)[OutfitIndex] = BodySize;
+		}
+		// Add or replace name
+		(*GetOutfitCollection()->Outfits.Name)[OutfitIndex] = Name;
+
+		// Add or replace pieces
+		int32& PiecesCount = (*GetOutfitCollection()->Outfits.PiecesCount)[OutfitIndex];
+		int32& PiecesStart = (*GetOutfitCollection()->Outfits.PiecesStart)[OutfitIndex];
+		if (PiecesStart && PiecesCount != INDEX_NONE)
+		{
+			GetOutfitCollection()->ManagedArrayCollection.RemoveElements(OutfitCollection::Group::Pieces, PiecesCount, PiecesStart);
+		}
+		PiecesCount = Pieces.Num();
+		if (PiecesCount)
+		{
+			PiecesStart = GetOutfitCollection()->ManagedArrayCollection.AddElements(PiecesCount, Group::Pieces);
+			int32 Index = PiecesStart;
+			for (const TPair<FGuid, FString>& Piece : Pieces)
+			{
+				(*GetOutfitCollection()->Pieces.Guid)[Index] = Piece.Key;
+				(*GetOutfitCollection()->Pieces.Name)[Index] = Piece.Value;
+				++Index;
+			}
+		}
+		else
+		{
+			check(PiecesStart == INDEX_NONE);
+		}
+	}
+
+	void FCollectionOutfitFacade::Append(const FCollectionOutfitConstFacade& Other)
+	{
+		check(IsValid());
+		if (!Other.IsValid())
+		{
+			return;
+		}
+
+		// Merge body sizes
+		const int32 NumBodySizes = Other.GetNumBodySizes();
+		TArray<int32> RemappedBodySize;
+		RemappedBodySize.SetNumUninitialized(NumBodySizes);
+		for (int32 BodySize = 0; BodySize < NumBodySizes; ++BodySize)
+		{
+			// Add or replace the named body size
+			RemappedBodySize[BodySize] = AddBodySize(
+				Other.GetBodySizeName(BodySize),
+				Other.GetBodySizeBodyPartsSkeletalMeshes(BodySize),
+				Other.GetBodySizeMeasurements(BodySize));
+		}
+
+		// Merge sized outfits
+		for (const FGuid& Guid : Other.GetOutfitGuids())
+		{
+			for (const int32 BodySize : Other.GetOutfitBodySizes(Guid))
+			{
+				AddOutfit(
+					Guid,
+					RemappedBodySize[BodySize],
+					Other.GetOutfitName(Guid, RemappedBodySize[BodySize]),
+					Other.GetOutfitPieces(Guid, RemappedBodySize[BodySize]));
+			}
+		}
+	}
+}	// namespace UE::Chaos::OutfitAsset
+
+#undef LOCTEXT_NAMESPACE
diff --git a/Engine/Plugins/Experimental/ChaosOutfitAsset/Source/ChaosOutfitAssetEngine/Private/ChaosClothAsset/Outfit.cpp b/Engine/Plugins/Experimental/ChaosOutfitAsset/Source/ChaosOutfitAssetEngine/Private/ChaosClothAsset/Outfit.cpp
--- a/Engine/Plugins/Experimental/ChaosOutfitAsset/Source/ChaosOutfitAssetEngine/Private/ChaosClothAsset/Outfit.cpp
+++ b/Engine/Plugins/Experimental/ChaosOutfitAsset/Source/ChaosOutfitAssetEngine/Private/ChaosClothAsset/Outfit.cpp
@@ -1,9 +1,11 @@
 // Copyright Epic Games, Inc. All Rights Reserved.
 
 #include "ChaosClothAsset/Outfit.h"
+#include "ChaosClothAsset/OutfitAsset.h"
 #include "Chaos/CollectionPropertyFacade.h"
 #include "ChaosClothAsset/ClothAsset.h"
 #include "ChaosClothAsset/ClothSimulationModel.h"
+#include "ChaosClothAsset/CollectionOutfitFacade.h"
 #include "ChaosClothAsset/OutfitAssetPrivate.h"
 #include "Rendering/SkeletalMeshLODModel.h"
 #include "Rendering/SkeletalMeshLODRenderData.h"
@@ -560,13 +562,13 @@ UChaosOutfit::UChaosOutfit(const FObjectInitializer& ObjectInitializer)
 		Pieces,
 		ReferenceSkeleton,
 		RenderData,
-		Materials);
+		Materials,
+		OutfitCollection);
 }
 
 UChaosOutfit::UChaosOutfit(FVTableHelper& Helper)
 	: Super(Helper)
-{
-}
+{}
 
 UChaosOutfit::~UChaosOutfit() = default;
 
@@ -576,6 +578,7 @@ void UChaosOutfit::Append(const UChaosOutfit& Other)
 		Other.GetReferenceSkeleton(),
 		Other.RenderData.Get(),
 		Other.Materials,
+		Other.OutfitCollection,
 		Other.Pieces.Num(),
 		[this, &Other](int32 Index) -> FChaosOutfitPiece
 		{
@@ -584,15 +587,22 @@ void UChaosOutfit::Append(const UChaosOutfit& Other)
 		Pieces,
 		ReferenceSkeleton,
 		RenderData.Get(),
-		Materials);
+		Materials,
+		OutfitCollection);
 }
 
 void UChaosOutfit::Add(const UChaosClothAssetBase& ClothAssetBase)
 {
+	const FManagedArrayCollection EmptyOutfitCollection;
+	const FManagedArrayCollection* const InOutfitCollection = ExactCast<const UChaosOutfitAsset>(&ClothAssetBase) ?
+		&ExactCast<const UChaosOutfitAsset>(&ClothAssetBase)->GetOutfitCollection() :
+		&EmptyOutfitCollection;
+
 	Merge(
 		ClothAssetBase.GetRefSkeleton(),
 		ClothAssetBase.GetResourceForRendering(),
 		ClothAssetBase.GetMaterials(),
+		*InOutfitCollection,
 		ClothAssetBase.GetNumClothSimulationModels(),
 		[this, &ClothAssetBase](int32 Index) -> FChaosOutfitPiece
 		{
@@ -601,25 +611,29 @@ void UChaosOutfit::Add(const UChaosClothAssetBase& ClothAssetBase)
 		Pieces,
 		ReferenceSkeleton,
 		RenderData.Get(),
-		Materials);
+		Materials,
+		OutfitCollection);
 }
 
 void UChaosOutfit::CopyTo(
 	TArray<FChaosOutfitPiece>& OutPieces,
 	FReferenceSkeleton& OutReferenceSkeleton,
 	TUniquePtr<FSkeletalMeshRenderData>& OutSkeletalMeshRenderData,
-	TArray<FSkeletalMaterial>& OutMaterials) const
+	TArray<FSkeletalMaterial>& OutMaterials,
+	FManagedArrayCollection& OutOutfitCollection) const
 {
 	Init(
 		OutPieces,
 		OutReferenceSkeleton,
 		OutSkeletalMeshRenderData,
-		OutMaterials);
+		OutMaterials,
+		OutOutfitCollection);
 
 	Merge(
 		ReferenceSkeleton,
 		RenderData.Get(),
 		Materials,
+		OutfitCollection,
 		Pieces.Num(),
 		[this](int32 Index) -> FChaosOutfitPiece
 		{
@@ -628,7 +642,8 @@ void UChaosOutfit::CopyTo(
 		OutPieces,
 		OutReferenceSkeleton,
 		OutSkeletalMeshRenderData.Get(),
-		OutMaterials);
+		OutMaterials,
+		OutOutfitCollection);
 }
 
 void UChaosOutfit::Serialize(FArchive& Ar)
@@ -644,12 +659,14 @@ void UChaosOutfit::Init(
 	TArray<FChaosOutfitPiece>& OutPieces,
 	FReferenceSkeleton& OutReferenceSkeleton,
 	TUniquePtr<RenderDataType>& OutSkeletalMeshRenderData,
-	TArray<FSkeletalMaterial>& OutMaterials)
+	TArray<FSkeletalMaterial>& OutMaterials,
+	FManagedArrayCollection& OutOutfitCollection)
 {
 	OutPieces.Reset();
 	OutReferenceSkeleton.Empty();
 	OutSkeletalMeshRenderData = MakeUnique<RenderDataType>();
 	OutMaterials.Reset();
+	OutOutfitCollection.Reset();
 
 	// Create a default valid reference skeleton
 	OutReferenceSkeleton.Empty(1);
@@ -665,6 +682,10 @@ void UChaosOutfit::Init(
 
 	// Set default values for the empty render data that can be usefully used in a merge
 	OutSkeletalMeshRenderData->bSupportRayTracing = false;
+
+	// Init outfit collection
+	UE::Chaos::OutfitAsset::FCollectionOutfitFacade OutfitFacade(OutOutfitCollection);
+	OutfitFacade.DefineSchema();
 }
 
 template<typename InLODRenderDataType, typename OutLODRenderDataType>
@@ -868,12 +889,14 @@ void UChaosOutfit::Merge(
 	const FReferenceSkeleton& InReferenceSkeleton,
 	const InRenderDataType* const InSkeletalMeshRenderData,
 	const TArray<FSkeletalMaterial>& InMaterials,
+	const FManagedArrayCollection& InOutfitCollection,
 	const int32 NumPieces,
 	TFunction<FChaosOutfitPiece(int32)>&& MakePiece,
 	TArray<FChaosOutfitPiece>& OutPieces,
 	FReferenceSkeleton& OutReferenceSkeleton,
 	OutRenderDataType* const OutSkeletalMeshRenderData,
-	TArray<FSkeletalMaterial>& OutMaterials)
+	TArray<FSkeletalMaterial>& OutMaterials,
+	FManagedArrayCollection& OutOutfitCollection)
 {
 	using namespace UE::Chaos::ClothAsset::Private;
 
@@ -933,4 +956,8 @@ void UChaosOutfit::Merge(
 		OutSkeletalMeshRenderData->NumNonOptionalLODs = NumLODs;
 		OutSkeletalMeshRenderData->bSupportRayTracing = OutSkeletalMeshRenderData->bSupportRayTracing || InSkeletalMeshRenderData->bSupportRayTracing;
 	}
+
+	// Merge outfit collection
+	UE::Chaos::OutfitAsset::FCollectionOutfitFacade OutfitFacade(OutOutfitCollection);
+	OutfitFacade.Append(InOutfitCollection);
 }
diff --git a/Engine/Plugins/Experimental/ChaosOutfitAsset/Source/ChaosOutfitAssetEngine/Private/ChaosClothAsset/OutfitAsset.cpp b/Engine/Plugins/Experimental/ChaosOutfitAsset/Source/ChaosOutfitAssetEngine/Private/ChaosClothAsset/OutfitAsset.cpp
--- a/Engine/Plugins/Experimental/ChaosOutfitAsset/Source/ChaosOutfitAssetEngine/Private/ChaosClothAsset/OutfitAsset.cpp
+++ b/Engine/Plugins/Experimental/ChaosOutfitAsset/Source/ChaosOutfitAssetEngine/Private/ChaosClothAsset/OutfitAsset.cpp
@@ -2,7 +2,8 @@
 
 #include "ChaosClothAsset/OutfitAsset.h"
 #include "ChaosClothAsset/ClothSimulationModel.h"
-#include "Engine/SkinnedAssetCommon.h"
+#include "ChaosClothAsset/CollectionOutfitFacade.h"
+#include "Engine/SkeletalMesh.h"
 #if WITH_EDITOR
 #include "DerivedDataCacheInterface.h"
 #endif
@@ -47,6 +48,8 @@ bool UChaosOutfitAsset::HasValidClothSimulationModels() const
 
 void UChaosOutfitAsset::Build(const TObjectPtr<const UChaosOutfit> Outfit)
 {
+	using namespace UE::Chaos::OutfitAsset;
+
 	// Stop the rendering
 	ReleaseResources();
 
@@ -60,7 +63,8 @@ void UChaosOutfitAsset::Build(const TObjectPtr<const UChaosOutfit> Outfit)
 			Pieces,
 			TempReferenceSkeleton,
 			TempSkeletalMeshRenderData,
-			Materials);
+			Materials,
+			OutfitCollection);
 	}
 	else
 	{
@@ -68,7 +72,20 @@ void UChaosOutfitAsset::Build(const TObjectPtr<const UChaosOutfit> Outfit)
 			Pieces,
 			TempReferenceSkeleton,
 			TempSkeletalMeshRenderData,
-			Materials);
+			Materials,
+			OutfitCollection);
+	}
+
+	// Populate the body sizes
+	const FCollectionOutfitConstFacade OutfitFacade(OutfitCollection);
+	const TArray<FString> OutfitBodyPartsSkeletalMeshes = OutfitFacade.GetOutfitBodyPartsSkeletalMeshes();
+	Bodies.Reset(OutfitBodyPartsSkeletalMeshes.Num());
+	for (const FString& OutfitBodyPartsSkeletalMesh : OutfitBodyPartsSkeletalMeshes)
+	{
+		if (USkeletalMesh* const Body = LoadObject<USkeletalMesh>(this, *OutfitBodyPartsSkeletalMesh))
+		{
+			Bodies.Emplace(Body);
+		}
 	}
 
 	// Set the new reference skeleton
diff --git a/Engine/Plugins/Experimental/ChaosOutfitAsset/Source/ChaosOutfitAssetEngine/Private/ChaosClothAsset/OutfitCollection.cpp b/Engine/Plugins/Experimental/ChaosOutfitAsset/Source/ChaosOutfitAssetEngine/Private/ChaosClothAsset/OutfitCollection.cpp
new file mode 100644
--- /dev/null
+++ b/Engine/Plugins/Experimental/ChaosOutfitAsset/Source/ChaosOutfitAssetEngine/Private/ChaosClothAsset/OutfitCollection.cpp
@@ -0,0 +1,58 @@
+// Copyright Epic Games, Inc. All Rights Reserved.
+
+#include "ChaosClothAsset/OutfitCollection.h"
+
+namespace UE::Chaos::OutfitAsset
+{
+	/**
+	 * Outfit collection reserved group and attribute names.
+	 */
+	namespace OutfitCollection
+	{
+		namespace Group
+		{
+			const FName Outfits(TEXT("OutfitCollection::Outfits"));
+			const FName Pieces(TEXT("OutfitCollection::Pieces"));
+			const FName BodySizes(TEXT("OutfitCollection::BodySizes"));
+			const FName BodyParts(TEXT("OutfitCollection::BodyParts"));
+			const FName Measurements(TEXT("OutfitCollection::Measurements"));
+		}
+		namespace Attribute
+		{
+			// Outfits group
+			namespace Outfits
+			{
+				const FName Guid(TEXT("Guid"));  // FGuid
+				const FName BodySize(TEXT("BodySize"));  // int32 -> BodySizes
+				const FName Name(TEXT("Name"));  // FString
+				const FName PiecesStart(TEXT("PiecesStart"));  // int32 -> OutfitPieces
+				const FName PiecesCount(TEXT("PiecesCount"));  // int32
+			}
+			// Pieces group
+			namespace Pieces
+			{
+				const FName Guid(TEXT("Guid"));  // FGuid
+				const FName Name(TEXT("Name"));  // FString
+			}
+			// BodySizes group
+			namespace BodySizes
+			{
+				const FName Name(TEXT("Name"));  // FString
+				const FName BodyPartsStart(TEXT("BodyPartsStart"));  // int32 -> BodyParts
+				const FName BodyPartsCount(TEXT("BodyPartsCount"));  // int32
+				const FName MeasurementsStart(TEXT("MeasurementsStart"));  // int32 -> Measurements
+				const FName MeasurementsCount(TEXT("MeasurementsCount"));  // int32
+			}
+			// BodyParts group
+			namespace BodyParts
+			{
+				const FName SkeletalMesh(TEXT("SkeletalMesh"));  // FString
+			}
+			// Measurements group
+			namespace Measurements
+			{
+				const FName Name(TEXT("Name"));  // FString
+			}
+		}
+	}
+}
diff --git a/Engine/Plugins/Experimental/ChaosOutfitAsset/Source/ChaosOutfitAssetEngine/Public/ChaosClothAsset/CollectionOutfitFacade.h b/Engine/Plugins/Experimental/ChaosOutfitAsset/Source/ChaosOutfitAssetEngine/Public/ChaosClothAsset/CollectionOutfitFacade.h
new file mode 100644
--- /dev/null
+++ b/Engine/Plugins/Experimental/ChaosOutfitAsset/Source/ChaosOutfitAssetEngine/Public/ChaosClothAsset/CollectionOutfitFacade.h
@@ -0,0 +1,134 @@
+// Copyright Epic Games, Inc. All Rights Reserved.
+
+#pragma once 
+
+#include "Templates/SharedPointer.h"
+#include "Templates/UniquePtr.h"
+
+template<typename T> class TManagedArray;
+struct FManagedArrayCollection;
+
+namespace UE::Chaos::OutfitAsset
+{
+	namespace Private
+	{
+		struct FOutfitCollection;
+	}
+
+	/**
+	 * Cloth outfit collection facade.
+	 * Const access (read only) version.
+	 */
+	class FCollectionOutfitConstFacade
+	{
+	public:
+		CHAOSOUTFITASSETENGINE_API explicit FCollectionOutfitConstFacade(const TSharedRef<const FManagedArrayCollection>& InManagedArrayCollection);
+		CHAOSOUTFITASSETENGINE_API explicit FCollectionOutfitConstFacade(const FManagedArrayCollection& InManagedArrayCollection);
+
+		FCollectionOutfitConstFacade() = delete;
+		FCollectionOutfitConstFacade(const FCollectionOutfitConstFacade&) = delete;
+		FCollectionOutfitConstFacade& operator=(const FCollectionOutfitConstFacade&) = delete;
+		FCollectionOutfitConstFacade(FCollectionOutfitConstFacade&&) = default;
+		FCollectionOutfitConstFacade& operator=(FCollectionOutfitConstFacade&&) = default;
+		CHAOSOUTFITASSETENGINE_API virtual ~FCollectionOutfitConstFacade();
+
+		/** Return whether the facade is defined on the collection. */
+		CHAOSOUTFITASSETENGINE_API bool IsValid() const;
+
+		/** Return the sized outfits GUID, one per sized outfit which can contain different sized outfits. */
+		CHAOSOUTFITASSETENGINE_API TArray<FGuid> GetOutfitGuids() const;
+
+		/** Return all unique body part skeletal meshes used by this outfit. */
+		CHAOSOUTFITASSETENGINE_API TArray<FString> GetOutfitBodyPartsSkeletalMeshes() const;
+
+		/** Return the body size available for the specified outfit GUID. */
+		CHAOSOUTFITASSETENGINE_API TArray<int32> GetOutfitBodySizes(const FGuid& Guid) const;
+
+		/** Return the name of the original asset making up the specified outfit body size. */
+		CHAOSOUTFITASSETENGINE_API const FString& GetOutfitName(const FGuid& Guid, const int32 BodySize) const;
+
+		/** Return the outfit pieces making up the specified outfit body size. */
+		CHAOSOUTFITASSETENGINE_API TMap<FGuid, FString> GetOutfitPieces(const FGuid& Guid, const int32 BodySize) const;
+
+		/** Return the number of body sizes. */
+		CHAOSOUTFITASSETENGINE_API int32 GetNumBodySizes() const;
+
+		/** Return whether this body size already exist. */
+		CHAOSOUTFITASSETENGINE_API bool HasBodySize(const FString& Name) const;
+
+		/** Return the specified body size name which is unique. */
+		CHAOSOUTFITASSETENGINE_API const FString& GetBodySizeName(const int32 BodySize) const;
+
+		/** Return the body part skeletal meshes path name for the specified body size. */
+		CHAOSOUTFITASSETENGINE_API TConstArrayView<FString> GetBodySizeBodyPartsSkeletalMeshes(const int32 BodySize) const;
+
+		/** Return the body part skeletal meshes path name for the specified body size. */
+		CHAOSOUTFITASSETENGINE_API TMap<FString, float> GetBodySizeMeasurements(const int32 BodySize) const;
+
+
+	protected:
+		const Private::FOutfitCollection* GetOutfitCollection() const { return OutfitCollection.Get(); }
+
+		TUniquePtr<const Private::FOutfitCollection> OutfitCollection;
+		TSharedPtr<const FManagedArrayCollection> ManagedArrayCollection;  // Only used to keep the shared ref alive
+	};
+
+	/**
+	 * Cloth outfit collection facade.
+	 * Non-const access (read/write) version.
+	 */
+	class FCollectionOutfitFacade final : public FCollectionOutfitConstFacade
+	{
+	public:
+		CHAOSOUTFITASSETENGINE_API explicit FCollectionOutfitFacade(const TSharedRef<FManagedArrayCollection>& InManagedArrayCollection);
+		CHAOSOUTFITASSETENGINE_API explicit FCollectionOutfitFacade(FManagedArrayCollection& InManagedArrayCollection);
+	
+		FCollectionOutfitFacade() = delete;
+		FCollectionOutfitFacade(const FCollectionOutfitFacade&) = delete;
+		FCollectionOutfitFacade& operator=(const FCollectionOutfitFacade&) = delete;
+		FCollectionOutfitFacade(FCollectionOutfitFacade&&) = default;
+		FCollectionOutfitFacade& operator=(FCollectionOutfitFacade&&) = default;
+		CHAOSOUTFITASSETENGINE_API virtual ~FCollectionOutfitFacade() override;
+
+		/** Add the Selection attributes to the underlying collection. */
+		CHAOSOUTFITASSETENGINE_API void DefineSchema();
+
+		/**
+		 * Add a new body size and returns its index.
+		 * Any pre-existing sizes data will be replaced.
+		 */
+		CHAOSOUTFITASSETENGINE_API int32 AddBodySize(
+			const FString& Name,
+			const TConstArrayView<FString> BodyPartsSkeletalMeshes,
+			const TMap<FString, float>& Measurements);
+
+		/**
+		 * Add a new sized outfit under the specified outfit GUID.
+		 * Any pre-existing sizes data will be replaced.
+		 */
+		CHAOSOUTFITASSETENGINE_API void AddOutfit(
+			const FGuid& Guid,
+			const int32 BodySize,
+			const FString& Name,
+			const TMap<FGuid, FString>& Pieces);
+
+		/**
+		 * Append an existing outfit facade to this collection.
+		 * Any pre-existing sizes data will be replaced.
+		 */
+		CHAOSOUTFITASSETENGINE_API void Append(const FCollectionOutfitConstFacade& Other);
+
+		/**
+		 * Append an existing outfit facade to this collection.
+		 * Any pre-existing sizes data will be replaced.
+		 */
+		void Append(const FManagedArrayCollection& InManagedArrayCollection)
+		{
+			FCollectionOutfitConstFacade Other(InManagedArrayCollection);
+			Append(Other);
+		}
+
+	private:
+		Private::FOutfitCollection* GetOutfitCollection() { return const_cast<Private::FOutfitCollection*>(OutfitCollection.Get()); }
+	};
+}
diff --git a/Engine/Plugins/Experimental/ChaosOutfitAsset/Source/ChaosOutfitAssetEngine/Public/ChaosClothAsset/Outfit.h b/Engine/Plugins/Experimental/ChaosOutfitAsset/Source/ChaosOutfitAssetEngine/Public/ChaosClothAsset/Outfit.h
--- a/Engine/Plugins/Experimental/ChaosOutfitAsset/Source/ChaosOutfitAssetEngine/Public/ChaosClothAsset/Outfit.h
+++ b/Engine/Plugins/Experimental/ChaosOutfitAsset/Source/ChaosOutfitAssetEngine/Public/ChaosClothAsset/Outfit.h
@@ -3,6 +3,7 @@
 #pragma once
 
 #include "Containers/Array.h"
+#include "GeometryCollection/ManagedArrayCollection.h"
 #include "Templates/SharedPointerFwd.h"
 #include "UObject/Object.h"
 #include "ReferenceSkeleton.h"
@@ -12,6 +13,7 @@
 class UChaosClothAssetBase;
 class FSkeletalMeshRenderData;
 class UPhysicsAsset;
+class USkeletalMesh;
 struct FChaosClothSimulationModel;
 struct FManagedArrayCollection;
 struct FSkeletalMaterial;
@@ -74,7 +76,7 @@ struct TStructOpsTypeTraits<FChaosOutfitPiece> : public TStructOpsTypeTraitsBase
 /**
  * Outfit class to handle the assembly of an outfit asset.
  */
-UCLASS(Experimental)
+UCLASS(Experimental, BlueprintType)
 class CHAOSOUTFITASSETENGINE_API UChaosOutfit final : public UObject
 {
 	GENERATED_BODY()
@@ -91,13 +93,23 @@ public:
 		TArray<FChaosOutfitPiece>& OutPieces,
 		FReferenceSkeleton& OutReferenceSkeleton,
 		TUniquePtr<FSkeletalMeshRenderData>& OutSkeletalMeshRenderData,
-		TArray<FSkeletalMaterial>& OutMaterials) const;
+		TArray<FSkeletalMaterial>& OutMaterials,
+		FManagedArrayCollection& OutOutfitCollection) const;
 
 	const TArray<FChaosOutfitPiece>& GetPieces() const
 	{
 		return Pieces;
 	}
 
+	const FManagedArrayCollection& GetOutfitCollection() const
+	{
+		return OutfitCollection;
+	}
+	FManagedArrayCollection& GetOutfitCollection()
+	{
+		return OutfitCollection;
+	}
+
 	const FReferenceSkeleton& GetReferenceSkeleton() const
 	{
 		return ReferenceSkeleton;
@@ -115,7 +127,8 @@ public:
 		TArray<FChaosOutfitPiece>& OutPieces,
 		FReferenceSkeleton& OutReferenceSkeleton,
 		TUniquePtr<RenderDataType>& OutSkeletalMeshRenderData,
-		TArray<FSkeletalMaterial>& OutMaterials);
+		TArray<FSkeletalMaterial>& OutMaterials,
+		FManagedArrayCollection& OutOutfitCollection);
 private:
 	template<typename InLODRenderDataType, typename OutLODRenderDataType>
 	static void MergeLODRenderDatas(
@@ -131,12 +144,14 @@ private:
 		const FReferenceSkeleton& InReferenceSkeleton,
 		const InRenderDataType* const InRenderData,
 		const TArray<FSkeletalMaterial>& InMaterials,
+		const FManagedArrayCollection& InOutfitCollection,
 		const int32 NumPieces,
 		TFunction<FChaosOutfitPiece(int32)>&& MakePiece,
 		TArray<FChaosOutfitPiece>& OutPieces,
 		FReferenceSkeleton& OutReferenceSkeleton,
 		OutRenderDataType* const OutRenderData,
-		TArray<FSkeletalMaterial>& OutMaterials);
+		TArray<FSkeletalMaterial>& OutMaterials,
+		FManagedArrayCollection& OutOutfitCollection);
 
 	UPROPERTY()
 	TArray<FChaosOutfitPiece> Pieces;
@@ -144,6 +159,9 @@ private:
 	UPROPERTY()
 	TArray<FSkeletalMaterial> Materials;
 
+	UPROPERTY()
+	FManagedArrayCollection OutfitCollection;
+
 	FReferenceSkeleton ReferenceSkeleton;
 
 	TUniquePtr<FRenderData> RenderData;
diff --git a/Engine/Plugins/Experimental/ChaosOutfitAsset/Source/ChaosOutfitAssetEngine/Public/ChaosClothAsset/OutfitAsset.h b/Engine/Plugins/Experimental/ChaosOutfitAsset/Source/ChaosOutfitAssetEngine/Public/ChaosClothAsset/OutfitAsset.h
--- a/Engine/Plugins/Experimental/ChaosOutfitAsset/Source/ChaosOutfitAssetEngine/Public/ChaosClothAsset/OutfitAsset.h
+++ b/Engine/Plugins/Experimental/ChaosOutfitAsset/Source/ChaosOutfitAssetEngine/Public/ChaosClothAsset/OutfitAsset.h
@@ -28,6 +28,11 @@ public:
 	virtual void PostLoad() override;
 	//~ End UObject interface
 
+	const FManagedArrayCollection& GetOutfitCollection() const
+	{
+		return OutfitCollection;
+	}
+
 private:
 	//~ Begin UChaosClothAssetBase interface
 	virtual bool HasValidClothSimulationModels() const override;
@@ -75,4 +80,10 @@ private:
 
 	UPROPERTY()
 	TArray<FChaosOutfitPiece> Pieces;
+
+	UPROPERTY()
+	TArray<TObjectPtr<USkeletalMesh>> Bodies;  // Only contains dependencies, populated from the outfit collection
+
+	UPROPERTY()
+	FManagedArrayCollection OutfitCollection;
 };
diff --git a/Engine/Plugins/Experimental/ChaosOutfitAsset/Source/ChaosOutfitAssetEngine/Public/ChaosClothAsset/OutfitCollection.h b/Engine/Plugins/Experimental/ChaosOutfitAsset/Source/ChaosOutfitAssetEngine/Public/ChaosClothAsset/OutfitCollection.h
new file mode 100644
--- /dev/null
+++ b/Engine/Plugins/Experimental/ChaosOutfitAsset/Source/ChaosOutfitAssetEngine/Public/ChaosClothAsset/OutfitCollection.h
@@ -0,0 +1,53 @@
+// Copyright Epic Games, Inc. All Rights Reserved.
+
+#pragma once
+
+#include "UObject/NameTypes.h"
+
+namespace UE::Chaos::OutfitAsset
+{
+	/**
+	 * Outfit collection reserved group and attribute names.
+	 */
+	namespace OutfitCollection
+	{
+		namespace Group
+		{
+			CHAOSOUTFITASSETENGINE_API extern const FName Outfits;
+			CHAOSOUTFITASSETENGINE_API extern const FName Pieces;
+			CHAOSOUTFITASSETENGINE_API extern const FName BodySizes;
+			CHAOSOUTFITASSETENGINE_API extern const FName BodyParts;
+			CHAOSOUTFITASSETENGINE_API extern const FName Measurements;
+		}
+		namespace Attribute
+		{
+			namespace Outfits
+			{
+				CHAOSOUTFITASSETENGINE_API extern const FName Guid;  // FGuid
+				CHAOSOUTFITASSETENGINE_API extern const FName BodySize;  // int32 -> BodySizes
+				CHAOSOUTFITASSETENGINE_API extern const FName Name;  // FString
+				CHAOSOUTFITASSETENGINE_API extern const FName PiecesStart;  // int32 -> Pieces
+				CHAOSOUTFITASSETENGINE_API extern const FName PiecesCount;  // int32
+			}
+			namespace Pieces
+			{
+				CHAOSOUTFITASSETENGINE_API extern const FName Guid;  // FGuid
+				CHAOSOUTFITASSETENGINE_API extern const FName Name;  // FString
+			}
+			namespace BodySizes
+			{
+				CHAOSOUTFITASSETENGINE_API extern const FName Name;  // FString
+				CHAOSOUTFITASSETENGINE_API extern const FName BodyPartsStart;  // int32 -> BodyParts
+				CHAOSOUTFITASSETENGINE_API extern const FName BodyPartsCount;  // int32
+			}
+			namespace BodyParts
+			{
+				CHAOSOUTFITASSETENGINE_API extern const FName SkeletalMesh;  // FString
+			}
+			namespace Measurements
+			{
+				CHAOSOUTFITASSETENGINE_API extern const FName Name;  // FString
+			}
+		}
+	}
+}
