Chaos Outfit Asset - Added a few minor improvments to the Outfit Collection change CL 41444817, that was depending on CL 41445018 to be committed, and which were suggested in code review:
- Added InvalidMeasurement constant value to BodyUserData, so that no measrument  is no longer set as a magic value of 0.
- Fixed a Outfit Facade API doc comment that was explaining a different function.
- Added a log message when replacing a named size in the Outfit Collection facade.
- Enabled reading of MetaHuman Measurements.

#rnx

diff --git a/Engine/Plugins/Experimental/ChaosOutfitAsset/Source/ChaosOutfitAssetDataflowNodes/ChaosOutfitAssetDataflowNodes.Build.cs b/Engine/Plugins/Experimental/ChaosOutfitAsset/Source/ChaosOutfitAssetDataflowNodes/ChaosOutfitAssetDataflowNodes.Build.cs
--- a/Engine/Plugins/Experimental/ChaosOutfitAsset/Source/ChaosOutfitAssetDataflowNodes/ChaosOutfitAssetDataflowNodes.Build.cs
+++ b/Engine/Plugins/Experimental/ChaosOutfitAsset/Source/ChaosOutfitAssetDataflowNodes/ChaosOutfitAssetDataflowNodes.Build.cs
@@ -18,7 +18,7 @@ public class ChaosOutfitAssetDataflowNodes : ModuleRules
 				"ChaosOutfitAssetEngine",
 				"DataflowCore",
 				"DataflowNodes",
-				"Engine",
+				"Engine"
 			}
 		);
 	}
diff --git a/Engine/Plugins/Experimental/ChaosOutfitAsset/Source/ChaosOutfitAssetDataflowNodes/Private/ChaosClothAsset/MakeSizedOutfitNode.cpp b/Engine/Plugins/Experimental/ChaosOutfitAsset/Source/ChaosOutfitAssetDataflowNodes/Private/ChaosClothAsset/MakeSizedOutfitNode.cpp
--- a/Engine/Plugins/Experimental/ChaosOutfitAsset/Source/ChaosOutfitAssetDataflowNodes/Private/ChaosClothAsset/MakeSizedOutfitNode.cpp
+++ b/Engine/Plugins/Experimental/ChaosOutfitAsset/Source/ChaosOutfitAssetDataflowNodes/Private/ChaosClothAsset/MakeSizedOutfitNode.cpp
@@ -4,6 +4,7 @@
 #include "ChaosClothAsset/CollectionOutfitFacade.h"
 #include "ChaosClothAsset/ClothAssetAnyType.h"
 #include "ChaosClothAsset/Outfit.h"
+#include "ChaosClothAsset/BodyUserData.h"
 #include "Engine/SkeletalMesh.h"
 
 #include UE_INLINE_GENERATED_CPP_BY_NAME(MakeSizedOutfitNode)
@@ -78,14 +79,26 @@ void FChaosOutfitAssetMakeSizedOutfitNode::Evaluate(UE::Dataflow::FContext& Cont
 
 				TArray<FString> BodyPartsSkeletalMeshes;
 				BodyPartsSkeletalMeshes.Reserve(SizedOutfitSource.SourceBodyParts.Num());
+				TMap<FString, float> Measurements;
 				for (const TObjectPtr<const USkeletalMesh>& SourceBodyPart : SizedOutfitSource.SourceBodyParts)
 				{
 					if (SourceBodyPart)
 					{
+						if (SourceBodyPart->GetAssetUserDataArray())
+						{
+							for (const UAssetUserData* const AssetUserData : *SourceBodyPart->GetAssetUserDataArray())
+							{
+								if (const UChaosOutfitAssetBodyUserData* const BodyAssetUserData =
+									Cast<const UChaosOutfitAssetBodyUserData>(AssetUserData))
+								{
+									Measurements.Append(BodyAssetUserData->Measurements);
+									break;
+								}
+							}
+						}
 						BodyPartsSkeletalMeshes.Emplace(SourceBodyPart->GetPathName());
 					}
 				}
-				TMap<FString, float> Measurements;  // TODO: Read measurements
 				const int32 BodySize = OutfitFacade.AddBodySize(BodySizeName, BodyPartsSkeletalMeshes, Measurements);
 
 				if (const TObjectPtr<const UChaosClothAssetBase> ClothAsset = SizedOutfitSource.SourceAsset)
diff --git a/Engine/Plugins/Experimental/ChaosOutfitAsset/Source/ChaosOutfitAssetEngine/Private/ChaosClothAsset/CollectionOutfitFacade.cpp b/Engine/Plugins/Experimental/ChaosOutfitAsset/Source/ChaosOutfitAssetEngine/Private/ChaosClothAsset/CollectionOutfitFacade.cpp
--- a/Engine/Plugins/Experimental/ChaosOutfitAsset/Source/ChaosOutfitAssetEngine/Private/ChaosClothAsset/CollectionOutfitFacade.cpp
+++ b/Engine/Plugins/Experimental/ChaosOutfitAsset/Source/ChaosOutfitAssetEngine/Private/ChaosClothAsset/CollectionOutfitFacade.cpp
@@ -1,6 +1,8 @@
 // Copyright Epic Games, Inc. All Rights Reserved.
 
 #include "ChaosClothAsset/CollectionOutfitFacade.h"
+#include "ChaosClothAsset/BodyUserData.h"
+#include "ChaosClothAsset/OutfitAssetPrivate.h"
 #include "ChaosClothAsset/OutfitCollection.h"
 #include "GeometryCollection/ManagedArrayCollection.h"
 
@@ -381,7 +383,7 @@ namespace UE::Chaos::OutfitAsset
 		else
 		{
 			// Clear all existing measurements
-			Values->Fill(0.f);
+			Values->Fill(UChaosOutfitAssetBodyUserData::InvalidMeasurement);
 		}
 		for (const TPair<FString, float>& Measurement : Measurements)
 		{
@@ -465,6 +467,9 @@ namespace UE::Chaos::OutfitAsset
 		RemappedBodySize.SetNumUninitialized(NumBodySizes);
 		for (int32 BodySize = 0; BodySize < NumBodySizes; ++BodySize)
 		{
+			UE_CLOG(HasBodySize(Other.GetBodySizeName(BodySize)), LogChaosOutfitAsset, Display,
+				TEXT("Outfit Collection Append operation caused body size [%s] to be replaced."), *Other.GetBodySizeName(BodySize));
+
 			// Add or replace the named body size
 			RemappedBodySize[BodySize] = AddBodySize(
 				Other.GetBodySizeName(BodySize),
diff --git a/Engine/Plugins/Experimental/ChaosOutfitAsset/Source/ChaosOutfitAssetEngine/Public/ChaosClothAsset/BodyUserData.h b/Engine/Plugins/Experimental/ChaosOutfitAsset/Source/ChaosOutfitAssetEngine/Public/ChaosClothAsset/BodyUserData.h
--- a/Engine/Plugins/Experimental/ChaosOutfitAsset/Source/ChaosOutfitAssetEngine/Public/ChaosClothAsset/BodyUserData.h
+++ b/Engine/Plugins/Experimental/ChaosOutfitAsset/Source/ChaosOutfitAssetEngine/Public/ChaosClothAsset/BodyUserData.h
@@ -14,6 +14,8 @@ class CHAOSOUTFITASSETENGINE_API UChaosOutfitAssetBodyUserData : public UAssetUs
 	GENERATED_BODY()
 
 public:
+	static inline constexpr float InvalidMeasurement = 0.f;
+
 	UPROPERTY(VisibleAnywhere, Category = "Measurements")
 	TMap<FString, float> Measurements;
 };
diff --git a/Engine/Plugins/Experimental/ChaosOutfitAsset/Source/ChaosOutfitAssetEngine/Public/ChaosClothAsset/CollectionOutfitFacade.h b/Engine/Plugins/Experimental/ChaosOutfitAsset/Source/ChaosOutfitAssetEngine/Public/ChaosClothAsset/CollectionOutfitFacade.h
--- a/Engine/Plugins/Experimental/ChaosOutfitAsset/Source/ChaosOutfitAssetEngine/Public/ChaosClothAsset/CollectionOutfitFacade.h
+++ b/Engine/Plugins/Experimental/ChaosOutfitAsset/Source/ChaosOutfitAssetEngine/Public/ChaosClothAsset/CollectionOutfitFacade.h
@@ -62,7 +62,7 @@ namespace UE::Chaos::OutfitAsset
 		/** Return the body part skeletal meshes path name for the specified body size. */
 		CHAOSOUTFITASSETENGINE_API TConstArrayView<FString> GetBodySizeBodyPartsSkeletalMeshes(const int32 BodySize) const;
 
-		/** Return the body part skeletal meshes path name for the specified body size. */
+		/** Return the body measurements stored for this specified body size. */
 		CHAOSOUTFITASSETENGINE_API TMap<FString, float> GetBodySizeMeasurements(const int32 BodySize) const;
 
 
