Chaos Outfit Asset - Fixed an issue where an existing Dataflow asset is overwritten when creating a new Outfit Asset that creates a new Dataflow of the same name.

#rnx

diff --git a/Engine/Plugins/Experimental/ChaosOutfitAsset/Source/ChaosOutfitAssetEditor/Private/ChaosOutfitAsset/OutfitAssetFactory.cpp b/Engine/Plugins/Experimental/ChaosOutfitAsset/Source/ChaosOutfitAssetEditor/Private/ChaosOutfitAsset/OutfitAssetFactory.cpp
--- a/Engine/Plugins/Experimental/ChaosOutfitAsset/Source/ChaosOutfitAssetEditor/Private/ChaosOutfitAsset/OutfitAssetFactory.cpp
+++ b/Engine/Plugins/Experimental/ChaosOutfitAsset/Source/ChaosOutfitAssetEditor/Private/ChaosOutfitAsset/OutfitAssetFactory.cpp
@@ -24,9 +24,14 @@ UObject* UChaosOutfitAssetFactory::FactoryCreateNew(UClass* Class, UObject* Pare
 	// Create a new Dataflow asset
 	const FString DataflowPath = FPackageName::GetLongPackagePath(OutfitAsset->GetOutermost()->GetName());
 	const FString OutfitAssetName = OutfitAsset->GetName();
-	const FString DataflowName = FString(TEXT("DF_")) + (OutfitAssetName.StartsWith(TEXT("OA_")) ? OutfitAssetName.RightChop(3) : OutfitAssetName);
-	const FString DataflowPackageName = FPaths::Combine(DataflowPath, DataflowName);
-
+	FString DataflowName = FString(TEXT("DF_")) + (OutfitAssetName.StartsWith(TEXT("OA_")) ? OutfitAssetName.RightChop(3) : OutfitAssetName);
+	FString DataflowPackageName = FPaths::Combine(DataflowPath, DataflowName);
+	if (FindPackage(nullptr, *DataflowPackageName))
+	{
+		// If a Dataflow asset already exists with this name, make a unique name from it to avoid clobbering it
+		MakeUniqueObjectName(nullptr, UPackage::StaticClass(), FName(DataflowPackageName)).ToString(DataflowPackageName);
+		DataflowName = FPaths::GetBaseFilename(DataflowPackageName);
+	}
 	UPackage* const DataflowPackage = CreatePackage(*DataflowPackageName);
 
 	// Load the cloth template into the new Dataflow asset
