Chaos Outfit Asset - Added sized outfit filtering node to remove all but the closely matched body size.

#rnx

diff --git a/Engine/Plugins/Experimental/ChaosOutfitAsset/Source/ChaosOutfitAssetDataflowNodes/Private/ChaosOutfitAsset/ChaosOutfitAssetDataflowNodesModule.cpp b/Engine/Plugins/Experimental/ChaosOutfitAsset/Source/ChaosOutfitAssetDataflowNodes/Private/ChaosOutfitAsset/ChaosOutfitAssetDataflowNodesModule.cpp
--- a/Engine/Plugins/Experimental/ChaosOutfitAsset/Source/ChaosOutfitAssetDataflowNodes/Private/ChaosOutfitAsset/ChaosOutfitAssetDataflowNodesModule.cpp
+++ b/Engine/Plugins/Experimental/ChaosOutfitAsset/Source/ChaosOutfitAssetDataflowNodes/Private/ChaosOutfitAsset/ChaosOutfitAssetDataflowNodesModule.cpp
@@ -1,6 +1,7 @@
 // Copyright Epic Games, Inc. All Rights Reserved.
 
 #include "ChaosOutfitAsset/ClothAssetAnyType.h"
+#include "ChaosOutfitAsset/FilterSizedOutfitNode.h"
 #include "ChaosOutfitAsset/GetClothAssetNode.h"
 #include "ChaosOutfitAsset/GetOutfitAssetNode.h"
 #include "ChaosOutfitAsset/MakeOutfitNode.h"
@@ -29,6 +30,7 @@ namespace UE::Chaos::OutfitAsset
 			DATAFLOW_NODE_REGISTER_CREATION_FACTORY_NODE_COLORS_BY_CATEGORY("Outfit", OutfitAssetNodeHeaderColor, OutfitAssetNodeBodyColor);
 			DATAFLOW_NODE_REGISTER_CREATION_FACTORY(FChaosGetClothAssetNode);
 			DATAFLOW_NODE_REGISTER_CREATION_FACTORY(FChaosGetOutfitAssetNode);
+			DATAFLOW_NODE_REGISTER_CREATION_FACTORY(FChaosOutfitAssetFilterSizedOutfitNode);
 			DATAFLOW_NODE_REGISTER_CREATION_FACTORY(FChaosOutfitAssetMakeOutfitNode);
 			DATAFLOW_NODE_REGISTER_CREATION_FACTORY(FChaosOutfitAssetMakeSizedOutfitNode);
 			DATAFLOW_NODE_REGISTER_CREATION_FACTORY(FChaosOutfitAssetMergeOutfitsNode);
diff --git a/Engine/Plugins/Experimental/ChaosOutfitAsset/Source/ChaosOutfitAssetDataflowNodes/Private/ChaosOutfitAsset/FilterSizedOutfitNode.cpp b/Engine/Plugins/Experimental/ChaosOutfitAsset/Source/ChaosOutfitAssetDataflowNodes/Private/ChaosOutfitAsset/FilterSizedOutfitNode.cpp
new file mode 100644
--- /dev/null
+++ b/Engine/Plugins/Experimental/ChaosOutfitAsset/Source/ChaosOutfitAssetDataflowNodes/Private/ChaosOutfitAsset/FilterSizedOutfitNode.cpp
@@ -0,0 +1,82 @@
+// Copyright Epic Games, Inc. All Rights Reserved.
+
+#include "ChaosOutfitAsset/FilterSizedOutfitNode.h"
+#include "ChaosOutfitAsset/Outfit.h"
+#include "ChaosOutfitAsset/CollectionOutfitFacade.h"
+#include "Engine/SkeletalMesh.h"
+
+#include UE_INLINE_GENERATED_CPP_BY_NAME(FilterSizedOutfitNode)
+
+#define LOCTEXT_NAMESPACE "ChaosOutfitAssetFilterSizedOutfitNode"
+
+FChaosOutfitAssetFilterSizedOutfitNode::FChaosOutfitAssetFilterSizedOutfitNode(const UE::Dataflow::FNodeParameters& InParam, FGuid InGuid)
+	: FDataflowNode(InParam, InGuid)
+{
+	RegisterInputConnection(&Outfit);
+	RegisterInputConnection(&SizeName);
+	RegisterInputConnection(&TargetBody);
+	RegisterOutputConnection(&Outfit, &Outfit);
+	RegisterOutputConnection(&SizeName, &SizeName);
+	RegisterOutputConnection(&OutfitCollection);
+}
+
+void FChaosOutfitAssetFilterSizedOutfitNode::Evaluate(UE::Dataflow::FContext& Context, const FDataflowOutput* Out) const
+{
+	using namespace UE::Chaos::OutfitAsset;
+
+	if (Out->IsA(&Outfit) || Out->IsA(&SizeName) || Out->IsA(&OutfitCollection))
+	{
+		if (const TObjectPtr<const UChaosOutfit>& InOutfit = GetValue(Context, &Outfit))
+		{
+			FCollectionOutfitConstFacade InOutfitFacade(InOutfit->GetOutfitCollection());
+			if (InOutfitFacade.IsValid())
+			{
+				FString InSizeName;
+
+				if (const TObjectPtr<const USkeletalMesh>& InTargetBody = GetValue(Context, &TargetBody))
+				{
+					// Selection by matching target body
+					const int32 ClosestBodySize = InOutfitFacade.FindClosestBodySize(*InTargetBody);
+					InSizeName = InOutfitFacade.GetBodySizeName(ClosestBodySize);  // Empty string when ClosestBodySize == INDEX_NONE
+				}
+				else
+				{
+					// Selection by name
+					InSizeName = GetValue(Context, &SizeName);
+					if (!InSizeName.IsEmpty() && !InOutfit->HasBodySize(InSizeName))
+					{
+						// Only show a warning when the size name is non empty
+						Context.Warning(FString::Printf(TEXT("The given body size [%s] doesn't exist in the input Outfit."), *SizeName), this, Out);
+						InSizeName.Reset();
+					}
+				}
+
+				if (!InSizeName.IsEmpty())
+				{
+					// Construct a new Outfit with only the input's pieces of the specified sizes
+					TObjectPtr<UChaosOutfit> OutOutfit = NewObject<UChaosOutfit>();
+					OutOutfit->Append(*InOutfit, InSizeName);
+
+					SetValue<TObjectPtr<const UChaosOutfit>>(Context, OutOutfit, &Outfit);
+					SetValue(Context, OutOutfit->GetOutfitCollection(), &OutfitCollection);
+					SetValue(Context, InSizeName, &SizeName);
+					return;
+				}
+			}
+			// No selection made, forward the input
+			SafeForwardInput(Context, &Outfit, &Outfit);
+			SetValue(Context, InOutfit->GetOutfitCollection(), &OutfitCollection);
+			SetValue(Context, FString(), &SizeName);
+		}
+		else
+		{
+			// No input Outfit, make up an empty output
+			TObjectPtr<UChaosOutfit> OutOutfit = NewObject<UChaosOutfit>();
+			SetValue<TObjectPtr<const UChaosOutfit>>(Context, OutOutfit, &Outfit);
+			SetValue(Context, OutOutfit->GetOutfitCollection(), &OutfitCollection);
+			SetValue(Context, FString(), &SizeName);
+		}
+	}
+}
+
+#undef LOCTEXT_NAMESPACE
diff --git a/Engine/Plugins/Experimental/ChaosOutfitAsset/Source/ChaosOutfitAssetDataflowNodes/Private/ChaosOutfitAsset/FilterSizedOutfitNode.h b/Engine/Plugins/Experimental/ChaosOutfitAsset/Source/ChaosOutfitAssetDataflowNodes/Private/ChaosOutfitAsset/FilterSizedOutfitNode.h
new file mode 100644
--- /dev/null
+++ b/Engine/Plugins/Experimental/ChaosOutfitAsset/Source/ChaosOutfitAssetDataflowNodes/Private/ChaosOutfitAsset/FilterSizedOutfitNode.h
@@ -0,0 +1,51 @@
+// Copyright Epic Games, Inc. All Rights Reserved.
+
+#pragma once 
+
+#include "Dataflow/DataflowNode.h"
+#include "GeometryCollection/ManagedArrayCollection.h"
+#include "FilterSizedOutfitNode.generated.h"
+
+class UChaosOutfit;
+class USkeletalMesh;
+
+/**
+ * Select a single size for the passed Outfit and filter out all non matching sizes.
+ */
+USTRUCT(Meta = (Experimental, DataflowOutfit))
+struct FChaosOutfitAssetFilterSizedOutfitNode final : public FDataflowNode
+{
+	GENERATED_USTRUCT_BODY()
+	DATAFLOW_NODE_DEFINE_INTERNAL(FChaosOutfitAssetFilterSizedOutfitNode, "FilterSizedOutfit", "Outfit", "Filter Sized Outfit")
+
+public:
+	FChaosOutfitAssetFilterSizedOutfitNode(const UE::Dataflow::FNodeParameters& InParam, FGuid InGuid = FGuid::NewGuid());
+
+private:
+	//~ Begin FDataflowNode interface
+	virtual void Evaluate(UE::Dataflow::FContext& Context, const FDataflowOutput* Out) const override;
+	//~ End FDataflowNode interface
+
+	/** The outfit to filter. */
+	UPROPERTY(Meta = (DataflowInput, DataflowOutput, DataflowPassthrough = "Outfit"))
+	TObjectPtr<const UChaosOutfit> Outfit;
+
+	/** The outfit collection output, provided for convenience as a view into the outfit object metadata. */
+	UPROPERTY(Meta = (DataflowOutput))
+	FManagedArrayCollection OutfitCollection;
+
+	/**
+	 * The name of the body size to use to filter.
+	 * If the input Size Name is empty, the output will be set to a name set from the body part measurements (if provided).
+	 * Requires the Target Body to be empty for the input Size Name to have an effect.
+	 */
+	UPROPERTY(EditAnywhere, Category = "FilterSizedOutfit", Meta = (DataflowInput, DataflowOutput, DataflowPassthrough = "SizeName", EditCondition = "TargetBody == nullptr"))
+	FString SizeName;
+
+	/**
+	 * The target body skeletal mesh containing the measurements to select the required size to use to filter.
+	 * The input Size Name is ignored when a valid target body is specified, and the selected size is returned as the output Size Name.
+	 */
+	UPROPERTY(EditAnywhere, Category = "FilterSizedOutfit", Meta = (DataflowInput))
+	TObjectPtr<const USkeletalMesh> TargetBody;
+};
diff --git a/Engine/Plugins/Experimental/ChaosOutfitAsset/Source/ChaosOutfitAssetDataflowNodes/Private/ChaosOutfitAsset/MakeSizedOutfitNode.cpp b/Engine/Plugins/Experimental/ChaosOutfitAsset/Source/ChaosOutfitAssetDataflowNodes/Private/ChaosOutfitAsset/MakeSizedOutfitNode.cpp
--- a/Engine/Plugins/Experimental/ChaosOutfitAsset/Source/ChaosOutfitAssetDataflowNodes/Private/ChaosOutfitAsset/MakeSizedOutfitNode.cpp
+++ b/Engine/Plugins/Experimental/ChaosOutfitAsset/Source/ChaosOutfitAssetDataflowNodes/Private/ChaosOutfitAsset/MakeSizedOutfitNode.cpp
@@ -4,7 +4,6 @@
 #include "ChaosOutfitAsset/CollectionOutfitFacade.h"
 #include "ChaosOutfitAsset/ClothAssetAnyType.h"
 #include "ChaosOutfitAsset/Outfit.h"
-#include "ChaosOutfitAsset/BodyUserData.h"
 #include "Engine/SkeletalMesh.h"
 
 #include UE_INLINE_GENERATED_CPP_BY_NAME(MakeSizedOutfitNode)
@@ -31,18 +30,6 @@ void FChaosOutfitAssetMakeSizedOutfitNode::Evaluate(UE::Dataflow::FContext& Cont
 		TObjectPtr<UChaosOutfit> OutOutfit = NewObject<UChaosOutfit>();
 		UE::Chaos::OutfitAsset::FCollectionOutfitFacade OutfitFacade(OutOutfit->GetOutfitCollection());
 
-		auto GetValidBodyPart = [](const FChaosSizedOutfitSource& SizedOutfitSource) -> const USkeletalMesh*
-			{
-				for (const TObjectPtr<const USkeletalMesh>& SourceBodyPart : SizedOutfitSource.SourceBodyParts)
-				{
-					if (SourceBodyPart)
-					{
-						return SourceBodyPart.Get();
-					}
-				}
-				return nullptr;
-			};
-
 		// Make a new GUID for this sized outfit asset
 		const FGuid OutfitGuid = FGuid::NewGuid();
 
@@ -52,82 +39,19 @@ void FChaosOutfitAssetMakeSizedOutfitNode::Evaluate(UE::Dataflow::FContext& Cont
 
 			for (const FChaosSizedOutfitSource& SizedOutfitSource : SizedOutfitSourceOrArray.Array)
 			{
-				// Check size, only display a warning for empty sizes when not empty of cloth assets, otherwise skip
-				FString BodySizeName = SizedOutfitSource.SizeName;
-				if (BodySizeName.IsEmpty())
-				{
-					// Use the skeletal mesh name if the size name is empty
-					if (const USkeletalMesh* SkeletalMesh = GetValidBodyPart(SizedOutfitSource))
-					{
-						BodySizeName = SkeletalMesh->GetName();
-					}
-					else if (SizedOutfitSource.SourceAsset && SizedOutfitSource.SourceAsset->GetNumClothSimulationModels())
-					{
-						BodySizeName = TEXT("Default");
-					}
-					else
-					{
-						continue;  // Empty size, skip
-					}
-				}
+				const FString BodySizeName = SizedOutfitSource.GetBodySizeName();
 
-				// Add body size
-				if (OutfitFacade.HasBodySize(BodySizeName))
+				if (!SizedOutfitSource.SourceAsset && !BodySizeName.IsEmpty())
 				{
-					Context.Warning(FString::Printf(TEXT("The body size [%s] already exist and has now been replaced."), *BodySizeName), this, Out);
+					const FString SourceAssetName = SizedOutfitSource.SourceAsset->GetName();
+					Context.Warning(FString::Printf(TEXT("Empty source asset [%s] specified for size [%s]."), *SourceAssetName, *BodySizeName), this, Out);
 				}
-
-				TArray<FString> BodyPartsSkeletalMeshes;
-				BodyPartsSkeletalMeshes.Reserve(SizedOutfitSource.SourceBodyParts.Num());
-				TMap<FString, float> Measurements;
-				for (const TObjectPtr<const USkeletalMesh>& SourceBodyPart : SizedOutfitSource.SourceBodyParts)
+				if (OutfitFacade.HasBodySize(BodySizeName))
 				{
-					if (SourceBodyPart)
-					{
-						if (SourceBodyPart->GetAssetUserDataArray())
-						{
-							for (const UAssetUserData* const AssetUserData : *SourceBodyPart->GetAssetUserDataArray())
-							{
-								if (const UChaosOutfitAssetBodyUserData* const BodyAssetUserData =
-									Cast<const UChaosOutfitAssetBodyUserData>(AssetUserData))
-								{
-									Measurements.Append(BodyAssetUserData->Measurements);
-									break;
-								}
-							}
-						}
-						BodyPartsSkeletalMeshes.Emplace(SourceBodyPart->GetPathName());
-					}
+					Context.Warning(FString::Printf(TEXT("The body size [%s] already existed and had to be overwritten."), *BodySizeName), this, Out);
 				}
-				const int32 BodySize = OutfitFacade.AddBodySize(BodySizeName, BodyPartsSkeletalMeshes, Measurements);
 
-				if (const TObjectPtr<const UChaosClothAssetBase> ClothAsset = SizedOutfitSource.SourceAsset)
-				{
-					const int32 NumPieces = ClothAsset->GetNumClothSimulationModels();
-					if (!NumPieces)
-					{
-						Context.Warning(FString::Printf(TEXT("Empty source asset [%s] specified for size [%s]."), *ClothAsset->GetName(), *BodySizeName), this, Out);
-						continue;
-					}
-					TMap<FGuid, FString> Pieces;
-					Pieces.Reserve(NumPieces);
-					for (int32 Index = 0; Index < NumPieces; ++Index)
-					{
-						Pieces.Emplace(
-							ClothAsset->GetAssetGuid(Index),
-							ClothAsset->GetClothSimulationModelName(Index).ToString());
-					}
-
-					// Add size metadata
-					OutfitFacade.AddOutfit(
-						OutfitGuid,
-						BodySize,
-						ClothAsset->GetName(),
-						Pieces);
-
-					// Add asset
-					OutOutfit->Add(*ClothAsset);
-				}
+				OutOutfit->Add(SizedOutfitSource, OutfitGuid);
 			}
 		}
 
diff --git a/Engine/Plugins/Experimental/ChaosOutfitAsset/Source/ChaosOutfitAssetDataflowNodes/Private/ChaosOutfitAsset/SizedOutfitSource.h b/Engine/Plugins/Experimental/ChaosOutfitAsset/Source/ChaosOutfitAssetDataflowNodes/Private/ChaosOutfitAsset/SizedOutfitSource.h
deleted file mode 100644
diff --git a/Engine/Plugins/Experimental/ChaosOutfitAsset/Source/ChaosOutfitAssetEngine/Private/ChaosOutfitAsset/CollectionOutfitFacade.cpp b/Engine/Plugins/Experimental/ChaosOutfitAsset/Source/ChaosOutfitAssetEngine/Private/ChaosOutfitAsset/CollectionOutfitFacade.cpp
--- a/Engine/Plugins/Experimental/ChaosOutfitAsset/Source/ChaosOutfitAssetEngine/Private/ChaosOutfitAsset/CollectionOutfitFacade.cpp
+++ b/Engine/Plugins/Experimental/ChaosOutfitAsset/Source/ChaosOutfitAssetEngine/Private/ChaosOutfitAsset/CollectionOutfitFacade.cpp
@@ -1,9 +1,11 @@
 // Copyright Epic Games, Inc. All Rights Reserved.
 
 #include "ChaosOutfitAsset/CollectionOutfitFacade.h"
+#include "ChaosClothAsset/ClothAssetBase.h"
 #include "ChaosOutfitAsset/BodyUserData.h"
 #include "ChaosOutfitAsset/OutfitAssetPrivate.h"
 #include "ChaosOutfitAsset/OutfitCollection.h"
+#include "Engine/SkeletalMesh.h"
 #include "GeometryCollection/ManagedArrayCollection.h"
 
 #define LOCTEXT_NAMESPACE "CollectionOutfitFacade"
@@ -131,22 +133,213 @@ namespace UE::Chaos::OutfitAsset
 			}
 
 			void DefineSchema()
+			{
+				if (!IsValid())
+				{
+					using namespace OutfitCollection;
+					Outfits.Guid = &ManagedArrayCollection.AddAttribute<FGuid>(Attribute::Outfits::Guid, Group::Outfits);
+					Outfits.BodySize = &ManagedArrayCollection.AddAttribute<int32>(Attribute::Outfits::BodySize, Group::Outfits, Group::BodySizes);
+					Outfits.Name = &ManagedArrayCollection.AddAttribute<FString>(Attribute::Outfits::Name, Group::Outfits);
+					Outfits.PiecesStart = &ManagedArrayCollection.AddAttribute<int32>(Attribute::Outfits::PiecesStart, Group::Outfits, Group::Pieces);
+					Outfits.PiecesCount = &ManagedArrayCollection.AddAttribute<int32>(Attribute::Outfits::PiecesCount, Group::Outfits);
+					Pieces.Guid = &ManagedArrayCollection.AddAttribute<FGuid>(Attribute::Pieces::Guid, Group::Pieces);
+					Pieces.Name = &ManagedArrayCollection.AddAttribute<FString>(Attribute::Pieces::Name, Group::Pieces);
+					BodySizes.Name = &ManagedArrayCollection.AddAttribute<FString>(Attribute::BodySizes::Name, Group::BodySizes);
+					BodySizes.BodyPartsStart = &ManagedArrayCollection.AddAttribute<int32>(Attribute::BodySizes::BodyPartsStart, Group::BodySizes, Group::BodyParts);
+					BodySizes.BodyPartsCount = &ManagedArrayCollection.AddAttribute<int32>(Attribute::BodySizes::BodyPartsCount, Group::BodySizes);
+					BodyParts.SkeletalMesh = &ManagedArrayCollection.AddAttribute<FString>(Attribute::BodyParts::SkeletalMesh, Group::BodyParts);
+					Measurements.Name = &ManagedArrayCollection.AddAttribute<FString>(Attribute::Measurements::Name, Group::Measurements);
+				}
+			}
+
+			void Reset()
 			{
 				using namespace OutfitCollection;
-				Outfits.Guid = &ManagedArrayCollection.AddAttribute<FGuid>(Attribute::Outfits::Guid, Group::Outfits);
-				Outfits.BodySize = &ManagedArrayCollection.AddAttribute<int32>(Attribute::Outfits::BodySize, Group::Outfits, Group::BodySizes);
-				Outfits.Name = &ManagedArrayCollection.AddAttribute<FString>(Attribute::Outfits::Name, Group::Outfits);
-				Outfits.PiecesStart = &ManagedArrayCollection.AddAttribute<int32>(Attribute::Outfits::PiecesStart, Group::Outfits, Group::Pieces);
-				Outfits.PiecesCount = &ManagedArrayCollection.AddAttribute<int32>(Attribute::Outfits::PiecesCount, Group::Outfits);
-				Pieces.Guid = &ManagedArrayCollection.AddAttribute<FGuid>(Attribute::Pieces::Guid, Group::Pieces);
-				Pieces.Name = &ManagedArrayCollection.AddAttribute<FString>(Attribute::Pieces::Name, Group::Pieces);
-				BodySizes.Name = &ManagedArrayCollection.AddAttribute<FString>(Attribute::BodySizes::Name, Group::BodySizes);
-				BodySizes.BodyPartsStart = &ManagedArrayCollection.AddAttribute<int32>(Attribute::BodySizes::BodyPartsStart, Group::BodySizes, Group::BodyParts);
-				BodySizes.BodyPartsCount = &ManagedArrayCollection.AddAttribute<int32>(Attribute::BodySizes::BodyPartsCount, Group::BodySizes);
-				BodyParts.SkeletalMesh = &ManagedArrayCollection.AddAttribute<FString>(Attribute::BodyParts::SkeletalMesh, Group::BodyParts);
-				Measurements.Name = &ManagedArrayCollection.AddAttribute<FString>(Attribute::Measurements::Name, Group::Measurements);
+				ManagedArrayCollection.EmptyGroup(Group::Outfits);
+				ManagedArrayCollection.EmptyGroup(Group::Pieces);
+				ManagedArrayCollection.EmptyGroup(Group::BodySizes);
+				ManagedArrayCollection.EmptyGroup(Group::BodyParts);
+				ManagedArrayCollection.EmptyGroup(Group::Measurements);
 			}
 		};
+
+		class FBodyMatchParameters
+		{
+		public:
+			enum class EBodyPlacement : uint8 { Top, Bottom, Full };
+
+			explicit FBodyMatchParameters(const TMap<FString, float>& Measurements)
+			{
+				Hip = Measurements.FindRef(TEXT("Hip"), InvalidMeasurement);
+				const float Bust = Measurements.FindRef(TEXT("Bust"), InvalidMeasurement);
+				const float Underbust = Measurements.FindRef(TEXT("Underbust"), InvalidMeasurement);
+				const float Waist = Measurements.FindRef(TEXT("Waist"), InvalidMeasurement);
+				const float NeckToWaist = Measurements.FindRef(TEXT("Neck to Waist"), InvalidMeasurement);
+				const float Rise = Measurements.FindRef(TEXT("Rise"), InvalidMeasurement);
+				const float Inseam = Measurements.FindRef(TEXT("Inseam"), InvalidMeasurement);
+
+				if (Hip != InvalidMeasurement &&
+					Bust != InvalidMeasurement &&
+					Underbust != InvalidMeasurement &&
+					Waist != InvalidMeasurement &&
+					NeckToWaist != InvalidMeasurement &&
+					Rise != InvalidMeasurement &&
+					Inseam != InvalidMeasurement &&
+					true)
+				{
+					// Protruding
+					bIsProtruding = Underbust > 5.f;
+
+					// Torso shape
+					const float HipBust = Hip / Bust;
+					const float WaistRatio = ((Bust + Hip) / 2.f) / Waist;
+
+					if (HipBust > 1.1f)
+					{
+						BodyShape = EBodyShape::Triangle;
+					}
+					else if (HipBust < 0.9f)
+					{
+						BodyShape = EBodyShape::InvertedTriangle;
+					}
+					else  // if (HipBust >= 0.9f && HipBust <= 1.1f)
+					{
+						if (WaistRatio > 1.3f)
+						{
+							BodyShape = EBodyShape::Circle;
+						}
+						else if (WaistRatio < 0.92f)
+						{
+							BodyShape = EBodyShape::Hourglass;
+						}
+						else  // if (WaistRatio >= 0.92f && WaistRatio <= 1.3f)
+						{
+							BodyShape = EBodyShape::Rectangle;
+						}
+					}
+
+					// Width
+					Width = FMath::Max(Hip, Waist, Bust) / 2.f;
+
+					// Proportions
+					ProportionTop = (NeckToWaist + (Rise / 2.f)) / Width;
+					ProportionBottom = Inseam / (Hip / 2.f);
+					ProportionFull = (NeckToWaist + (Rise / 2.f) + Inseam) / Width;
+				}
+			}
+
+			static int32 Score(const FBodyMatchParameters& Garment, const FBodyMatchParameters& Body, const EBodyPlacement BodyPlacement = EBodyPlacement::Full)
+			{
+				if (Garment.BodyShape == EBodyShape::Invalid || Body.BodyShape == EBodyShape::Invalid)
+				{
+					return 0;
+				}
+				int32 Score = 0;
+
+				// Protruding
+				if (Garment.bIsProtruding == Body.bIsProtruding)
+				{
+					Score += 1000;
+				}
+
+				// Body shape
+				switch (Garment.BodyShape)
+				{
+				case EBodyShape::Invalid: break;
+				case EBodyShape::Rectangle:  // Hourglass > Inverted Tri > Tri > Circle
+					switch (Body.BodyShape)
+					{
+					case EBodyShape::Invalid: break;
+					case EBodyShape::Rectangle: Score += 500; break;
+					case EBodyShape::Hourglass: Score += 400; break;
+					case EBodyShape::InvertedTriangle: Score += 300; break;
+					case EBodyShape::Triangle: Score += 200; break;
+					case EBodyShape::Circle: Score += 100; break;
+					}
+					break;
+				case EBodyShape::Circle:  // Rectangle > Triangle > InvertedTriangle > Hourglass
+					switch (Body.BodyShape)
+					{
+					case EBodyShape::Invalid: break;
+					case EBodyShape::Circle: Score += 500; break;
+					case EBodyShape::Rectangle: Score += 400; break;
+					case EBodyShape::Triangle: Score += 300; break;
+					case EBodyShape::InvertedTriangle: Score += 200; break;
+					case EBodyShape::Hourglass: Score += 100; break;
+					}
+					break;
+				case EBodyShape::Hourglass:  // Triangle > Rectangle > InvertedTriangle > Circle
+					switch (Body.BodyShape)
+					{
+					case EBodyShape::Invalid: break;
+					case EBodyShape::Hourglass: Score += 500; break;
+					case EBodyShape::Triangle: Score += 400; break;
+					case EBodyShape::Rectangle: Score += 300; break;
+					case EBodyShape::InvertedTriangle: Score += 200; break;
+					case EBodyShape::Circle: Score += 100; break;
+					}
+					break;
+				case EBodyShape::Triangle:  // Hourglass > Circle > Rectangle > InvertedTriangle
+					switch (Body.BodyShape)
+					{
+					case EBodyShape::Invalid: break;
+					case EBodyShape::Triangle: Score += 500; break;
+					case EBodyShape::Hourglass: Score += 400; break;
+					case EBodyShape::Circle: Score += 300; break;
+					case EBodyShape::Rectangle: Score += 200; break;
+					case EBodyShape::InvertedTriangle: Score += 100; break;
+					}
+					break;
+				case EBodyShape::InvertedTriangle:  // Rectangle > Circle > Hourglass > Triangle
+					switch (Body.BodyShape)
+					{
+					case EBodyShape::Invalid: break;
+					case EBodyShape::InvertedTriangle: Score += 500; break;
+					case EBodyShape::Rectangle: Score += 400; break;
+					case EBodyShape::Circle: Score += 300; break;
+					case EBodyShape::Hourglass: Score += 200; break;
+					case EBodyShape::Triangle: Score += 100; break;
+					}
+					break;
+				}
+
+				// Proportion
+				float ProportionDelta;
+				switch (BodyPlacement)
+				{
+				default:
+				case EBodyPlacement::Full: ProportionDelta = FMath::Abs(Garment.ProportionFull - Body.ProportionFull); break;
+				case EBodyPlacement::Top: ProportionDelta = FMath::Abs(Garment.ProportionTop - Body.ProportionTop); break;
+				case EBodyPlacement::Bottom: ProportionDelta = FMath::Abs(Garment.ProportionBottom - Body.ProportionBottom); break;
+				}
+				if (ProportionDelta <= 0.1f)
+				{
+					Score += 50;
+				}
+				else if (ProportionDelta <= 0.5f)
+				{
+					Score += (int32)(5.f / ProportionDelta);  // 50 to 10
+				}
+				// else Score += 0
+
+				// Hip comparison
+				Score -= (int32)(5.f * FMath::Abs(Garment.Hip - Body.Hip));  // 5 points per cm penalty
+
+				return Score;
+			}
+
+		private:
+			enum class EBodyShape : uint8 { Invalid, Rectangle, Circle, Hourglass, Triangle, InvertedTriangle };
+			static inline const float InvalidMeasurement = UChaosOutfitAssetBodyUserData::InvalidMeasurement;
+
+			float Hip = 0.f;
+			float Width = 0.f;
+			float ProportionTop = 0.f;
+			float ProportionBottom = 0.f;
+			float ProportionFull = 0.f;
+			EBodyShape BodyShape = EBodyShape::Invalid;
+			bool bIsProtruding = false;
+		};
 	}  // namespace Private
 
 	FCollectionOutfitConstFacade::FCollectionOutfitConstFacade(const TSharedRef<const FManagedArrayCollection>& InManagedArrayCollection)
@@ -219,6 +412,11 @@ namespace UE::Chaos::OutfitAsset
 		return EmptyString;
 	}
 
+	TConstArrayView<FGuid> FCollectionOutfitConstFacade::GetOutfitPiecesGuids() const
+	{
+		return TConstArrayView<FGuid>(GetOutfitCollection()->Pieces.Guid->GetData(), GetOutfitCollection()->Pieces.Guid->Num());
+	}
+
 	TMap<FGuid, FString> FCollectionOutfitConstFacade::GetOutfitPieces(const FGuid& Guid, const int32 BodySize) const
 	{
 		check(IsValid());
@@ -270,10 +468,57 @@ namespace UE::Chaos::OutfitAsset
 		return GetOutfitCollection()->BodySizes.Name->Num();
 	}
 
-	bool FCollectionOutfitConstFacade::HasBodySize(const FString& Name) const
+	int32 FCollectionOutfitConstFacade::FindBodySize(const FString& Name) const
+	{
+		return GetOutfitCollection()->BodySizes.Name->Find(Name);
+	}
+
+	int32 FCollectionOutfitConstFacade::FindClosestBodySize(const TMap<FString, float>& Measurements) const
 	{
 		check(IsValid());
-		return GetOutfitCollection()->BodySizes.Name->Find(Name) != INDEX_NONE;
+		const int32 NumBodySizes = GetNumBodySizes();
+
+		switch (NumBodySizes)
+		{
+		case 0: return INDEX_NONE;
+		case 1: return 0;
+		default: break;
+		}
+
+		const Private::FBodyMatchParameters BodyMatchParameters(Measurements);
+
+		int32 ClosestBodySize = INDEX_NONE;
+		int32 BestScore = TNumericLimits<int32>::Min();
+
+		for (int32 BodySize = 0; BodySize < NumBodySizes; ++BodySize)
+		{
+			const Private::FBodyMatchParameters GarmentMatchParameters(GetBodySizeMeasurements(BodySize));
+
+			const int32 Score = Private::FBodyMatchParameters::Score(GarmentMatchParameters, BodyMatchParameters);
+			if (Score > BestScore)
+			{
+				ClosestBodySize = BodySize;
+				BestScore = Score;
+			}
+		}
+
+		return ClosestBodySize;
+	}
+
+	int32 FCollectionOutfitConstFacade::FindClosestBodySize(const USkeletalMesh& BodyPart) const
+	{
+		if (BodyPart.GetAssetUserDataArray())
+		{
+			for (const UAssetUserData* const AssetUserData : *BodyPart.GetAssetUserDataArray())
+			{
+				if (const UChaosOutfitAssetBodyUserData* const BodyAssetUserData =
+					Cast<const UChaosOutfitAssetBodyUserData>(AssetUserData))
+				{
+					return FindClosestBodySize(BodyAssetUserData->Measurements);
+				}
+			}
+		}
+		return INDEX_NONE;
 	}
 
 	const FString& FCollectionOutfitConstFacade::GetBodySizeName(int32 BodySize) const
@@ -340,6 +585,12 @@ namespace UE::Chaos::OutfitAsset
 		check(IsValid());
 	}
 
+	void FCollectionOutfitFacade::Reset()
+	{
+		GetOutfitCollection()->Reset();
+		check(IsValid());
+	}
+
 	int32 FCollectionOutfitFacade::AddBodySize(const FString& Name, const TConstArrayView<FString> BodyPartsSkeletalMeshes, const TMap<FString, float>& Measurements)
 	{
 		using namespace OutfitCollection;
@@ -431,7 +682,7 @@ namespace UE::Chaos::OutfitAsset
 		// Add or replace pieces
 		int32& PiecesCount = (*GetOutfitCollection()->Outfits.PiecesCount)[OutfitIndex];
 		int32& PiecesStart = (*GetOutfitCollection()->Outfits.PiecesStart)[OutfitIndex];
-		if (PiecesStart && PiecesCount != INDEX_NONE)
+		if (PiecesStart != INDEX_NONE && PiecesCount)
 		{
 			GetOutfitCollection()->ManagedArrayCollection.RemoveElements(OutfitCollection::Group::Pieces, PiecesCount, PiecesStart);
 		}
@@ -453,7 +704,26 @@ namespace UE::Chaos::OutfitAsset
 		}
 	}
 
-	void FCollectionOutfitFacade::Append(const FCollectionOutfitConstFacade& Other)
+	void FCollectionOutfitFacade::AddOutfit(const FGuid& Guid, const int32 BodySize, const UChaosClothAssetBase& ClothAssetBase)
+	{
+		TMap<FGuid, FString> PieceInfos;
+		const int32 NumPieces = ClothAssetBase.GetNumClothSimulationModels();
+		PieceInfos.Reserve(NumPieces);
+		for (int32 Index = 0; Index < NumPieces; ++Index)
+		{
+			PieceInfos.Emplace(
+				ClothAssetBase.GetAssetGuid(Index),
+				ClothAssetBase.GetClothSimulationModelName(Index).ToString());
+		}
+
+		AddOutfit(
+			Guid,
+			BodySize,
+			ClothAssetBase.GetName(),
+			PieceInfos);
+	}
+
+	void FCollectionOutfitFacade::Append(const FCollectionOutfitConstFacade& Other, int32 SelectedBodySize)
 	{
 		check(IsValid());
 		if (!Other.IsValid())
@@ -462,16 +732,29 @@ namespace UE::Chaos::OutfitAsset
 		}
 
 		// Merge body sizes
-		const int32 NumBodySizes = Other.GetNumBodySizes();
+		int32 NumBodySizes;
+		if (SelectedBodySize == INDEX_NONE)
+		{
+			// Merge all the sizes
+			NumBodySizes = Other.GetNumBodySizes();
+			SelectedBodySize = 0;
+		}
+		else
+		{
+			// Merge the selected size
+			NumBodySizes = 1;
+			check(SelectedBodySize >= 0 && SelectedBodySize < Other.GetNumBodySizes());
+		}
+
 		TArray<int32> RemappedBodySize;
 		RemappedBodySize.SetNumUninitialized(NumBodySizes);
-		for (int32 BodySize = 0; BodySize < NumBodySizes; ++BodySize)
+		for (int32 BodySize = SelectedBodySize; BodySize < SelectedBodySize + NumBodySizes; ++BodySize)
 		{
 			UE_CLOG(HasBodySize(Other.GetBodySizeName(BodySize)), LogChaosOutfitAsset, Display,
 				TEXT("Outfit Collection Append operation caused body size [%s] to be replaced."), *Other.GetBodySizeName(BodySize));
 
 			// Add or replace the named body size
-			RemappedBodySize[BodySize] = AddBodySize(
+			RemappedBodySize[BodySize - SelectedBodySize] = AddBodySize(
 				Other.GetBodySizeName(BodySize),
 				Other.GetBodySizeBodyPartsSkeletalMeshes(BodySize),
 				Other.GetBodySizeMeasurements(BodySize));
@@ -482,11 +765,14 @@ namespace UE::Chaos::OutfitAsset
 		{
 			for (const int32 BodySize : Other.GetOutfitBodySizes(Guid))
 			{
-				AddOutfit(
-					Guid,
-					RemappedBodySize[BodySize],
-					Other.GetOutfitName(Guid, RemappedBodySize[BodySize]),
-					Other.GetOutfitPieces(Guid, RemappedBodySize[BodySize]));
+				if (NumBodySizes > 1 || BodySize == SelectedBodySize)
+				{
+					AddOutfit(
+						Guid,
+						RemappedBodySize[BodySize - SelectedBodySize],
+						Other.GetOutfitName(Guid, BodySize),
+						Other.GetOutfitPieces(Guid, BodySize));
+				}
 			}
 		}
 	}
diff --git a/Engine/Plugins/Experimental/ChaosOutfitAsset/Source/ChaosOutfitAssetEngine/Private/ChaosOutfitAsset/Outfit.cpp b/Engine/Plugins/Experimental/ChaosOutfitAsset/Source/ChaosOutfitAssetEngine/Private/ChaosOutfitAsset/Outfit.cpp
--- a/Engine/Plugins/Experimental/ChaosOutfitAsset/Source/ChaosOutfitAssetEngine/Private/ChaosOutfitAsset/Outfit.cpp
+++ b/Engine/Plugins/Experimental/ChaosOutfitAsset/Source/ChaosOutfitAssetEngine/Private/ChaosOutfitAsset/Outfit.cpp
@@ -4,9 +4,13 @@
 #include "Chaos/CollectionPropertyFacade.h"
 #include "ChaosClothAsset/ClothAsset.h"
 #include "ChaosClothAsset/ClothSimulationModel.h"
+#include "ChaosOutfitAsset/BodyUserData.h"
 #include "ChaosOutfitAsset/CollectionOutfitFacade.h"
 #include "ChaosOutfitAsset/OutfitAsset.h"
 #include "ChaosOutfitAsset/OutfitAssetPrivate.h"
+#include "ChaosOutfitAsset/OutfitCollection.h"
+#include "ChaosOutfitAsset/SizedOutfitSource.h"
+#include "Engine/SkeletalMesh.h"
 #include "Rendering/SkeletalMeshLODModel.h"
 #include "Rendering/SkeletalMeshLODRenderData.h"
 #include "Rendering/SkeletalMeshRenderData.h"
@@ -374,23 +378,31 @@ namespace UE::Chaos::OutfitAsset::Private
 		return VertexColors;
 	}
 
-	static TArray<FMeshToMeshVertData> GetClothMappingData(const TArray<FSkelMeshRenderSection>& RenderSections)
+	static TArray<FMeshToMeshVertData> GetClothMappingData(const TArray<FSkelMeshRenderSection>& RenderSections, const TArray<FGuid>* AssetGuids = nullptr)
 	{
 		TArray<FMeshToMeshVertData> MappingData;
 		int32 NumMappingData = 0;
 		for (const FSkelMeshRenderSection& RenderSection : RenderSections)
 		{
-			for (const TArray<FMeshToMeshVertData>& ClothMappingDataLOD : RenderSection.ClothMappingDataLODs)
+			// TODO: Fix ClothMappings
+			//if (!AssetGuids || AssetGuids->Find(RenderSection.ClothingData.AssetGuid) != INDEX_NONE)
 			{
-				NumMappingData += ClothMappingDataLOD.Num();
+				for (const TArray<FMeshToMeshVertData>& ClothMappingDataLOD : RenderSection.ClothMappingDataLODs)
+				{
+					NumMappingData += ClothMappingDataLOD.Num();
+				}
 			}
 		}
 		MappingData.Reserve(NumMappingData);
 		for (const FSkelMeshRenderSection& RenderSection : RenderSections)
 		{
-			for (const TArray<FMeshToMeshVertData>& ClothMappingDataLOD : RenderSection.ClothMappingDataLODs)
+			// TODO: Fix ClothMappings
+			//if (!AssetGuids || AssetGuids->Find(RenderSection.ClothingData.AssetGuid) != INDEX_NONE)
 			{
-				MappingData.Append(ClothMappingDataLOD);
+				for (const TArray<FMeshToMeshVertData>& ClothMappingDataLOD : RenderSection.ClothMappingDataLODs)
+				{
+					MappingData.Append(ClothMappingDataLOD);
+				}
 			}
 		}
 		return MappingData;
@@ -399,22 +411,37 @@ namespace UE::Chaos::OutfitAsset::Private
 	static TArray<FClothBufferIndexMapping> GetClothBufferIndexMappings(
 		const FSkeletalMeshVertexClothBuffer& ClothVertexBuffer,
 		const TArray<FSkelMeshRenderSection>& RenderSections,
-		const int32 VertexOffset = 0)
+		const int32 VertexOffset = 0,
+		const TArray<FGuid>* AssetGuids = nullptr)
 	{
-		TArray<FClothBufferIndexMapping> ClothBufferIndexMappings = ClothVertexBuffer.GetClothIndexMapping();
-		if (!ClothBufferIndexMappings.Num())
+		int32 NumRenderSections = 0;
+		for (const FSkelMeshRenderSection& RenderSection : RenderSections)
+		{
+			if (!AssetGuids || AssetGuids->Find(RenderSection.ClothingData.AssetGuid) != INDEX_NONE)
+			{
+				++NumRenderSections;
+			}
+		}
+
+		TArray<FClothBufferIndexMapping> ClothBufferIndexMappings;
+		if (!ClothVertexBuffer.GetClothIndexMapping().Num())
 		{
 			// Add empty mapping as some of the LOD sections have clothing
-			ClothBufferIndexMappings.SetNumZeroed(RenderSections.Num());  // FClothBufferIndexMapping has no default constructor
+			ClothBufferIndexMappings.SetNumZeroed(NumRenderSections);  // FClothBufferIndexMapping has no default constructor
 		}
 		else
 		{
 			// Update the mappings index/offset
-			for (FClothBufferIndexMapping& ClothBufferIndexMapping : ClothBufferIndexMappings)
+			for (int32 SectionIndex = 0; SectionIndex < RenderSections.Num(); ++SectionIndex)
 			{
-				ClothBufferIndexMapping.BaseVertexIndex += VertexOffset;
-				ClothBufferIndexMapping.MappingOffset += VertexOffset;
-				// LODBiasStride stays the same since the number of mapping for this section hasn't changed
+				const FSkelMeshRenderSection& RenderSection = RenderSections[SectionIndex];
+				if (!AssetGuids || AssetGuids->Find(RenderSection.ClothingData.AssetGuid) != INDEX_NONE)
+				{
+					FClothBufferIndexMapping& ClothBufferIndexMapping = ClothBufferIndexMappings.Add_GetRef(ClothVertexBuffer.GetClothIndexMapping()[SectionIndex]);
+					ClothBufferIndexMapping.BaseVertexIndex += VertexOffset;
+					ClothBufferIndexMapping.MappingOffset += VertexOffset;
+					// LODBiasStride stays the same since the number of mapping for this section hasn't changed
+				}
 			}
 		}
 		return ClothBufferIndexMappings;
@@ -430,6 +457,27 @@ namespace UE::Chaos::OutfitAsset::Private
 		ReferenceSkeleton.EnsureParentsExistAndSort(OutBoneIndices);
 		OutBoneIndices.Shrink();
 	}
+
+	template<typename T UE_REQUIRES(std::is_same_v<T, UChaosClothAssetBase> || std::is_same_v<T, TArray<FChaosOutfitPiece>>)>
+	FChaosOutfitPiece Clone(const T& Pieces, int32 ModelIndex)
+	{
+		return FChaosOutfitPiece(Pieces, ModelIndex);
+	}
+
+	template<typename T UE_REQUIRES(std::is_same_v<T, UChaosClothAssetBase> || std::is_same_v<T, TArray<FChaosOutfitPiece>>)>
+	FGuid GetAssetGuid(const T& Pieces, int32 ModelIndex);
+
+	template<>
+	FGuid GetAssetGuid<UChaosClothAssetBase>(const UChaosClothAssetBase& Pieces, int32 ModelIndex)
+	{
+		return Pieces.GetAssetGuid(ModelIndex);
+	}
+
+	template<>
+	FGuid GetAssetGuid<TArray<FChaosOutfitPiece>>(const TArray<FChaosOutfitPiece>& Pieces, int32 ModelIndex)
+	{
+		return Pieces[ModelIndex].AssetGuid;
+	}
 }  // namespace UE::Chaos::OutfitAsset::Private
 
 FChaosOutfitPiece::FChaosOutfitPiece()
@@ -572,18 +620,16 @@ UChaosOutfit::UChaosOutfit(FVTableHelper& Helper)
 
 UChaosOutfit::~UChaosOutfit() = default;
 
-void UChaosOutfit::Append(const UChaosOutfit& Other)
+void UChaosOutfit::Append(const UChaosOutfit& Other, const FString& BodySizeNameFilter)
 {
 	Merge(
 		Other.GetReferenceSkeleton(),
 		Other.RenderData.Get(),
 		Other.Materials,
 		Other.OutfitCollection,
+		Other.Pieces,
 		Other.Pieces.Num(),
-		[this, &Other](int32 Index) -> FChaosOutfitPiece
-		{
-			return FChaosOutfitPiece(Other.Pieces[Index]);
-		},
+		BodySizeNameFilter,
 		Pieces,
 		ReferenceSkeleton,
 		RenderData.Get(),
@@ -593,21 +639,39 @@ void UChaosOutfit::Append(const UChaosOutfit& Other)
 
 void UChaosOutfit::Add(const UChaosClothAssetBase& ClothAssetBase)
 {
-	const FManagedArrayCollection EmptyOutfitCollection;
-	const FManagedArrayCollection* const InOutfitCollection = ExactCast<const UChaosOutfitAsset>(&ClothAssetBase) ?
-		&ExactCast<const UChaosOutfitAsset>(&ClothAssetBase)->GetOutfitCollection() :
-		&EmptyOutfitCollection;
+	using namespace UE::Chaos::OutfitAsset;
+
+	// Retrieve the input outfit collection
+	TOptional<FManagedArrayCollection> InOutfitCollection;
+	const FManagedArrayCollection* InOutfitCollectionPtr = nullptr;
+
+	if (const UChaosClothAsset* ClothAsset = ExactCast<const UChaosClothAsset>(&ClothAssetBase))
+	{
+		// Build some outfit metadata for this cloth asset
+		FCollectionOutfitFacade InOutfitFacade(InOutfitCollection.Emplace());
+		InOutfitFacade.DefineSchema();
+
+		const FGuid OutfitGuid = FGuid::NewGuid();
+		const int32 BodySize = InOutfitFacade.FindOrAddBodySize(DefaultBodySize.ToString());  // TODO: Should the default body size be part of the schema?
+		InOutfitFacade.AddOutfit(OutfitGuid, BodySize, ClothAssetBase);
 
+		InOutfitCollectionPtr = &InOutfitCollection.GetValue();
+	}
+	else if(const UChaosOutfitAsset* const OutfitAsset = ExactCast<const UChaosOutfitAsset>(&ClothAssetBase))
+	{
+		InOutfitCollectionPtr = &OutfitAsset->GetOutfitCollection();
+	}
+	check(InOutfitCollectionPtr);
+
+	// Add cloth/outfit asset
 	Merge(
 		ClothAssetBase.GetRefSkeleton(),
 		ClothAssetBase.GetResourceForRendering(),
 		ClothAssetBase.GetMaterials(),
-		*InOutfitCollection,
+		*InOutfitCollectionPtr,
+		ClothAssetBase,
 		ClothAssetBase.GetNumClothSimulationModels(),
-		[this, &ClothAssetBase](int32 Index) -> FChaosOutfitPiece
-		{
-			return FChaosOutfitPiece(ClothAssetBase, Index);
-		},
+		{},
 		Pieces,
 		ReferenceSkeleton,
 		RenderData.Get(),
@@ -615,6 +679,78 @@ void UChaosOutfit::Add(const UChaosClothAssetBase& ClothAssetBase)
 		OutfitCollection);
 }
 
+void UChaosOutfit::Add(const FChaosSizedOutfitSource& SizedOutfitSource, const FGuid& OutfitGuid)
+{
+	using namespace UE::Chaos::OutfitAsset;
+
+	// Add body size
+	FManagedArrayCollection InOutfitCollection;
+	FCollectionOutfitFacade InOutfitFacade(InOutfitCollection);
+	InOutfitFacade.DefineSchema();
+
+	int32 BodySize = INDEX_NONE;
+
+	const FString BodySizeName = SizedOutfitSource.GetBodySizeName();
+	if (!BodySizeName.IsEmpty())
+	{
+		TArray<FString> BodyPartsSkeletalMeshes;
+		BodyPartsSkeletalMeshes.Reserve(SizedOutfitSource.SourceBodyParts.Num());
+		TMap<FString, float> Measurements;
+		for (const TObjectPtr<const USkeletalMesh>& SourceBodyPart : SizedOutfitSource.SourceBodyParts)
+		{
+			if (SourceBodyPart)
+			{
+				if (SourceBodyPart->GetAssetUserDataArray())
+				{
+					for (const UAssetUserData* const AssetUserData : *SourceBodyPart->GetAssetUserDataArray())  // Not using the simpler GetAssetUserDataOfClass here because it is not const
+					{
+						if (const UChaosOutfitAssetBodyUserData* const BodyAssetUserData =
+							Cast<const UChaosOutfitAssetBodyUserData>(AssetUserData))
+						{
+							Measurements.Append(BodyAssetUserData->Measurements);
+							break;
+						}
+					}
+				}
+				BodyPartsSkeletalMeshes.Emplace(SourceBodyPart->GetPathName());
+			}
+		}
+		BodySize = InOutfitFacade.AddBodySize(BodySizeName, BodyPartsSkeletalMeshes, Measurements);
+	}
+	else
+	{
+		BodySize = InOutfitFacade.FindOrAddBodySize(DefaultBodySize.ToString());
+	}
+
+	// Add source asset pieces
+	if (const UChaosClothAssetBase* const ClothAssetBase = SizedOutfitSource.SourceAsset)
+	{
+		// Move the entire set of outfit pieces under the new size/GUID
+		InOutfitFacade.AddOutfit(OutfitGuid, BodySize, *ClothAssetBase);
+
+		// Add cloth/outfit asset
+		Merge(
+			ClothAssetBase->GetRefSkeleton(),
+			ClothAssetBase->GetResourceForRendering(),
+			ClothAssetBase->GetMaterials(),
+			InOutfitCollection,
+			*ClothAssetBase,
+			ClothAssetBase->GetNumClothSimulationModels(),
+			{},
+			Pieces,
+			ReferenceSkeleton,
+			RenderData.Get(),
+			Materials,
+			OutfitCollection);
+	}
+	else
+	{
+		// Only add the size
+		FCollectionOutfitFacade OutfitFacade(OutfitCollection);
+		OutfitFacade.Append(InOutfitFacade);
+	}
+}
+
 void UChaosOutfit::CopyTo(
 	TArray<FChaosOutfitPiece>& OutPieces,
 	FReferenceSkeleton& OutReferenceSkeleton,
@@ -634,11 +770,9 @@ void UChaosOutfit::CopyTo(
 		RenderData.Get(),
 		Materials,
 		OutfitCollection,
+		Pieces,
 		Pieces.Num(),
-		[this](int32 Index) -> FChaosOutfitPiece
-		{
-			return FChaosOutfitPiece(Pieces[Index]);
-		},
+		{},
 		OutPieces,
 		OutReferenceSkeleton,
 		OutSkeletalMeshRenderData.Get(),
@@ -646,6 +780,12 @@ void UChaosOutfit::CopyTo(
 		OutOutfitCollection);
 }
 
+bool UChaosOutfit::HasBodySize(const FString& SizeName) const
+{
+	UE::Chaos::OutfitAsset::FCollectionOutfitConstFacade OutfitFacade(OutfitCollection);
+	return OutfitFacade.IsValid() && OutfitFacade.HasBodySize(SizeName);
+}
+
 void UChaosOutfit::Serialize(FArchive& Ar)
 {
 	Super::Serialize(Ar);
@@ -835,12 +975,12 @@ void UChaosOutfit::MergeLODRenderDatas(
 	if (bHasClothingData)
 	{
 		TArray<FMeshToMeshVertData> ClothMappingData = GetClothMappingData(OutLODRenderData.RenderSections);
-		ClothMappingData.Append(GetClothMappingData(LODRenderData.RenderSections));
+		ClothMappingData.Append(GetClothMappingData(LODRenderData.RenderSections, &AssetGuids));
 
 		TArray<FClothBufferIndexMapping> ClothIndexMappings =
 			GetClothBufferIndexMappings(OutLODRenderData.ClothVertexBuffer, OutLODRenderData.RenderSections);
 		ClothIndexMappings.Append(
-			GetClothBufferIndexMappings(LODRenderData.ClothVertexBuffer, LODRenderData.RenderSections, VertexOffset));
+			GetClothBufferIndexMappings(LODRenderData.ClothVertexBuffer, LODRenderData.RenderSections, VertexOffset, &AssetGuids));
 
 		OutLODRenderData.ClothVertexBuffer.Init(ClothMappingData, ClothIndexMappings);
 	}
@@ -853,29 +993,31 @@ void UChaosOutfit::MergeLODRenderDatas(
 	for (const FSkelMeshRenderSection& RenderSection : LODRenderData.RenderSections)
 	{
 		const int32 PieceIndex = AssetGuids.Find(RenderSection.ClothingData.AssetGuid);
-
-		FSkelMeshRenderSection& OutRenderSection = OutLODRenderData.RenderSections.AddDefaulted_GetRef();
-
-		OutRenderSection.MaterialIndex = RenderSection.MaterialIndex + MaterialOffset;
-		OutRenderSection.BaseIndex = RenderSection.BaseIndex + IndexOffset;
-		OutRenderSection.NumTriangles = RenderSection.NumTriangles;
-		OutRenderSection.bRecomputeTangent = RenderSection.bRecomputeTangent;
-		OutRenderSection.bCastShadow = RenderSection.bCastShadow;
-		OutRenderSection.bVisibleInRayTracing = RenderSection.bVisibleInRayTracing;
-		OutRenderSection.RecomputeTangentsVertexMaskChannel = RenderSection.RecomputeTangentsVertexMaskChannel;
-		OutRenderSection.BaseVertexIndex = RenderSection.BaseVertexIndex + VertexOffset;
-		OutRenderSection.ClothMappingDataLODs = RenderSection.ClothMappingDataLODs;
-		OutRenderSection.BoneMap = RenderSection.BoneMap;
-		OutRenderSection.NumVertices = RenderSection.NumVertices;
-		OutRenderSection.MaxBoneInfluences = RenderSection.MaxBoneInfluences;
-		OutRenderSection.CorrespondClothAssetIndex = PieceIndex;
-		OutRenderSection.ClothingData = RenderSection.ClothingData;
-		OutRenderSection.DuplicatedVerticesBuffer.Init(0, TMap<int32, TArray<int32>>());
-		OutRenderSection.bDisabled = RenderSection.bDisabled;
-
-		for (FBoneIndexType& BoneIndex : OutRenderSection.BoneMap)
+		if (PieceIndex != INDEX_NONE)
 		{
-			BoneIndex = BoneMap[BoneIndex];
+			FSkelMeshRenderSection& OutRenderSection = OutLODRenderData.RenderSections.AddDefaulted_GetRef();
+
+			OutRenderSection.MaterialIndex = RenderSection.MaterialIndex + MaterialOffset;
+			OutRenderSection.BaseIndex = RenderSection.BaseIndex + IndexOffset;
+			OutRenderSection.NumTriangles = RenderSection.NumTriangles;
+			OutRenderSection.bRecomputeTangent = RenderSection.bRecomputeTangent;
+			OutRenderSection.bCastShadow = RenderSection.bCastShadow;
+			OutRenderSection.bVisibleInRayTracing = RenderSection.bVisibleInRayTracing;
+			OutRenderSection.RecomputeTangentsVertexMaskChannel = RenderSection.RecomputeTangentsVertexMaskChannel;
+			OutRenderSection.BaseVertexIndex = RenderSection.BaseVertexIndex + VertexOffset;
+			OutRenderSection.ClothMappingDataLODs = RenderSection.ClothMappingDataLODs;
+			OutRenderSection.BoneMap = RenderSection.BoneMap;
+			OutRenderSection.NumVertices = RenderSection.NumVertices;
+			OutRenderSection.MaxBoneInfluences = RenderSection.MaxBoneInfluences;
+			OutRenderSection.CorrespondClothAssetIndex = PieceIndex;
+			OutRenderSection.ClothingData = RenderSection.ClothingData;
+			OutRenderSection.DuplicatedVerticesBuffer.Init(0, TMap<int32, TArray<int32>>());
+			OutRenderSection.bDisabled = RenderSection.bDisabled;
+
+			for (FBoneIndexType& BoneIndex : OutRenderSection.BoneMap)
+			{
+				BoneIndex = BoneMap[BoneIndex];
+			}
 		}
 	}
 
@@ -884,27 +1026,66 @@ void UChaosOutfit::MergeLODRenderDatas(
 	MergeBones(ReferenceSkeleton, BoneMap, LODRenderData.RequiredBones, OutLODRenderData.RequiredBones);
 }
 
-template<typename InRenderDataType, typename OutRenderDataType>
+template<typename InRenderDataType, typename OutRenderDataType, typename PiecesType>
 void UChaosOutfit::Merge(
 	const FReferenceSkeleton& InReferenceSkeleton,
 	const InRenderDataType* const InSkeletalMeshRenderData,
 	const TArray<FSkeletalMaterial>& InMaterials,
 	const FManagedArrayCollection& InOutfitCollection,
+	const PiecesType& Pieces,
 	const int32 NumPieces,
-	TFunction<FChaosOutfitPiece(int32)>&& MakePiece,
+	const FString& BodySizeNameFilter,
 	TArray<FChaosOutfitPiece>& OutPieces,
 	FReferenceSkeleton& OutReferenceSkeleton,
 	OutRenderDataType* const OutSkeletalMeshRenderData,
 	TArray<FSkeletalMaterial>& OutMaterials,
 	FManagedArrayCollection& OutOutfitCollection)
 {
-	using namespace UE::Chaos::OutfitAsset::Private;
+	using namespace UE::Chaos::OutfitAsset;
 
 	check(InSkeletalMeshRenderData);
 
+	// Filter the body size pieces to merge, or select all pieces
+	TArray<FGuid> AssetGuidsFilter;
+
+	const FCollectionOutfitConstFacade InOutfitFacade(InOutfitCollection);
+	FCollectionOutfitFacade OutOutfitFacade(OutOutfitCollection);
+	check(InOutfitFacade.IsValid());
+	check(OutOutfitFacade.IsValid());
+	
+	if (!BodySizeNameFilter.IsEmpty())
+	{
+		// Filter to the requested body size
+		const int32 BodySize = InOutfitFacade.FindBodySize(BodySizeNameFilter);
+		if (BodySize == INDEX_NONE)
+		{
+			UE_LOG(LogChaosOutfitAsset, Display, TEXT("The requested body size [%s] isn't available to merge."), *BodySizeNameFilter);
+			return;
+		}
+
+		// Find all pieces assigned to this body size
+		TMap<FGuid, FString> OutfitPieces;
+		for (const FGuid& Guid : InOutfitFacade.GetOutfitGuids())
+		{
+			OutfitPieces.Append(InOutfitFacade.GetOutfitPieces(Guid, BodySize));
+		}
+		OutfitPieces.GetKeys(AssetGuidsFilter);
+
+		// Merge collection, but only the requested body size
+		OutOutfitFacade.Append(InOutfitFacade, BodySize);
+	}
+	else
+	{
+		// Add all pieces to the merge
+		AssetGuidsFilter = InOutfitFacade.GetOutfitPiecesGuids();
+
+		// Merge the entire collection
+		OutOutfitFacade.Append(InOutfitFacade);
+	}
+
 	// Merge skeletons
 	TArray<int32> BoneMap;
-	MergeSkeletons(InReferenceSkeleton, OutReferenceSkeleton, BoneMap);
+	Private::MergeSkeletons(InReferenceSkeleton, OutReferenceSkeleton, BoneMap);
 
 	// Merge pieces
 	TArray<FGuid> AssetGuids;
@@ -913,9 +1094,14 @@ void UChaosOutfit::Merge(
 
 	for (int32 PieceIndex = 0; PieceIndex < NumPieces; ++PieceIndex)
 	{
-		FChaosOutfitPiece& Piece = OutPieces.Add_GetRef(MakePiece(PieceIndex));
-		Piece.RemapBoneIndices(BoneMap);
-		AssetGuids.AddUnique(Piece.AssetGuid);
+		if (AssetGuidsFilter.Contains(Private::GetAssetGuid(Pieces, PieceIndex)))
+		{
+			FChaosOutfitPiece& Piece = OutPieces.Add_GetRef(Private::Clone(Pieces, PieceIndex));
+			Piece.RemapBoneIndices(BoneMap);
+			// TODO: Fix case where a piece already exists in the outfit as to not mess up the LODRenderData piece indices
+			UE_CLOG(AssetGuids.Contains(Piece.AssetGuid), LogChaosOutfitAsset, Warning, TEXT("Piece [%s] already exists in this Outfit!."), *Piece.Name.ToString());
+			AssetGuids.Add(Piece.AssetGuid);
+		}
 	}
 
 	// Merge materials
@@ -956,8 +1142,4 @@ void UChaosOutfit::Merge(
 		OutSkeletalMeshRenderData->NumNonOptionalLODs = NumLODs;
 		OutSkeletalMeshRenderData->bSupportRayTracing = OutSkeletalMeshRenderData->bSupportRayTracing || InSkeletalMeshRenderData->bSupportRayTracing;
 	}
-
-	// Merge outfit collection
-	UE::Chaos::OutfitAsset::FCollectionOutfitFacade OutfitFacade(OutOutfitCollection);
-	OutfitFacade.Append(InOutfitCollection);
 }
diff --git a/Engine/Plugins/Experimental/ChaosOutfitAsset/Source/ChaosOutfitAssetEngine/Private/ChaosOutfitAsset/OutfitCollection.cpp b/Engine/Plugins/Experimental/ChaosOutfitAsset/Source/ChaosOutfitAssetEngine/Private/ChaosOutfitAsset/OutfitCollection.cpp
--- a/Engine/Plugins/Experimental/ChaosOutfitAsset/Source/ChaosOutfitAssetEngine/Private/ChaosOutfitAsset/OutfitCollection.cpp
+++ b/Engine/Plugins/Experimental/ChaosOutfitAsset/Source/ChaosOutfitAssetEngine/Private/ChaosOutfitAsset/OutfitCollection.cpp
@@ -55,4 +55,6 @@ namespace UE::Chaos::OutfitAsset
 			}
 		}
 	}  // namespace OutfitCollection
+
+	const FName DefaultBodySize(TEXT("Default"));
 }  // namespace UE::Chaos::OutfitAsset
diff --git a/Engine/Plugins/Experimental/ChaosOutfitAsset/Source/ChaosOutfitAssetEngine/Private/ChaosOutfitAsset/SizedOutfitSource.cpp b/Engine/Plugins/Experimental/ChaosOutfitAsset/Source/ChaosOutfitAssetEngine/Private/ChaosOutfitAsset/SizedOutfitSource.cpp
new file mode 100644
--- /dev/null
+++ b/Engine/Plugins/Experimental/ChaosOutfitAsset/Source/ChaosOutfitAssetEngine/Private/ChaosOutfitAsset/SizedOutfitSource.cpp
@@ -0,0 +1,28 @@
+// Copyright Epic Games, Inc. All Rights Reserved.
+
+#include "ChaosOutfitAsset/SizedOutfitSource.h"
+#include "ChaosOutfitAsset/OutfitCollection.h"
+#include "ChaosClothAsset/ClothAssetBase.h"
+
+FString FChaosSizedOutfitSource::GetBodySizeName() const
+{
+	FString BodySizeName = SizeName;
+	if (BodySizeName.IsEmpty())
+	{
+		// Use the skeletal mesh name if the size name is empty
+		if (const TObjectPtr<const USkeletalMesh>* const SkeletalMesh = SourceBodyParts.FindByPredicate(
+			[](const TObjectPtr<const USkeletalMesh>& SourceBodyPart) -> bool
+			{
+				return !!SourceBodyPart;
+			}))
+		{
+			BodySizeName = SkeletalMesh->GetName();
+		}
+		else if (SourceAsset && SourceAsset->GetNumClothSimulationModels())
+		{
+			// Return "Default" if the source asset is valid and no body parts have been provided
+			BodySizeName = UE::Chaos::OutfitAsset::DefaultBodySize.ToString();
+		}
+	}
+	return BodySizeName;
+}
diff --git a/Engine/Plugins/Experimental/ChaosOutfitAsset/Source/ChaosOutfitAssetEngine/Public/ChaosOutfitAsset/CollectionOutfitFacade.h b/Engine/Plugins/Experimental/ChaosOutfitAsset/Source/ChaosOutfitAssetEngine/Public/ChaosOutfitAsset/CollectionOutfitFacade.h
--- a/Engine/Plugins/Experimental/ChaosOutfitAsset/Source/ChaosOutfitAssetEngine/Public/ChaosOutfitAsset/CollectionOutfitFacade.h
+++ b/Engine/Plugins/Experimental/ChaosOutfitAsset/Source/ChaosOutfitAssetEngine/Public/ChaosOutfitAsset/CollectionOutfitFacade.h
@@ -6,6 +6,8 @@
 #include "Templates/UniquePtr.h"
 
 template<typename T> class TManagedArray;
+class UChaosClothAssetBase;
+class USkeletalMesh;
 struct FManagedArrayCollection;
 
 namespace UE::Chaos::OutfitAsset
@@ -47,6 +49,9 @@ namespace UE::Chaos::OutfitAsset
 		/** Return the name of the original asset making up the specified outfit body size. */
 		CHAOSOUTFITASSETENGINE_API const FString& GetOutfitName(const FGuid& Guid, const int32 BodySize) const;
 
+		/** Return all outfit pieces asset GUIDs. */
+		CHAOSOUTFITASSETENGINE_API TConstArrayView<FGuid> GetOutfitPiecesGuids() const;
+
 		/** Return the outfit pieces making up the specified outfit body size. */
 		CHAOSOUTFITASSETENGINE_API TMap<FGuid, FString> GetOutfitPieces(const FGuid& Guid, const int32 BodySize) const;
 
@@ -54,7 +59,19 @@ namespace UE::Chaos::OutfitAsset
 		CHAOSOUTFITASSETENGINE_API int32 GetNumBodySizes() const;
 
 		/** Return whether this body size already exist. */
-		CHAOSOUTFITASSETENGINE_API bool HasBodySize(const FString& Name) const;
+		bool HasBodySize(const FString& Name) const
+		{
+			return FindBodySize(Name) != INDEX_NONE;
+		}
+
+		/** Return the body size index for the specified name, or INDEX_NONE if it doesn't exist in this outfit. */
+		CHAOSOUTFITASSETENGINE_API int32 FindBodySize(const FString& Name) const;
+
+		/** Return the closest size to the specified measurements, or INDEX_NONE if this outfit has no sizes or the measurements are incomplete. */
+		CHAOSOUTFITASSETENGINE_API int32 FindClosestBodySize(const TMap<FString, float>& Measurements) const;
+
+		/** Return the closest size to the specified body, or INDEX_NONE if this outfit has no sizes or the body has no measurement data. */
+		CHAOSOUTFITASSETENGINE_API int32 FindClosestBodySize(const USkeletalMesh& BodyPart) const;
 
 		/** Return the specified body size name which is unique. */
 		CHAOSOUTFITASSETENGINE_API const FString& GetBodySizeName(const int32 BodySize) const;
@@ -93,6 +110,9 @@ namespace UE::Chaos::OutfitAsset
 		/** Add the Selection attributes to the underlying collection. */
 		CHAOSOUTFITASSETENGINE_API void DefineSchema();
 
+		/** Remove all outfit and size information from this collection. */
+		CHAOSOUTFITASSETENGINE_API void Reset();
+
 		/**
 		 * Add a new body size and returns its index.
 		 * Any pre-existing sizes data will be replaced.
@@ -102,6 +122,13 @@ namespace UE::Chaos::OutfitAsset
 			const TConstArrayView<FString> BodyPartsSkeletalMeshes,
 			const TMap<FString, float>& Measurements);
 
+		/** Find or add a new body size and returns its index. */
+		int32 FindOrAddBodySize(const FString& Name)
+		{
+			const int32 BodySize = FindBodySize(Name);
+			return BodySize != INDEX_NONE ? BodySize : AddBodySize(Name, {}, {});
+		}
+
 		/**
 		 * Add a new sized outfit under the specified outfit GUID.
 		 * Any pre-existing sizes data will be replaced.
@@ -109,14 +136,15 @@ namespace UE::Chaos::OutfitAsset
 		CHAOSOUTFITASSETENGINE_API void AddOutfit(
 			const FGuid& Guid,
 			const int32 BodySize,
-			const FString& Name,
-			const TMap<FGuid, FString>& Pieces);
+			const UChaosClothAssetBase& ClothAssetBase);
 
 		/**
 		 * Append an existing outfit facade to this collection.
 		 * Any pre-existing sizes data will be replaced.
+		 * @param Other The source outfit collection to get the data from.
+		 * @param SelectedBodySize The body size of the source outfit to append, or INDEX_NONE to append all sizes.
 		 */
-		CHAOSOUTFITASSETENGINE_API void Append(const FCollectionOutfitConstFacade& Other);
+		CHAOSOUTFITASSETENGINE_API void Append(const FCollectionOutfitConstFacade& Other, int32 SelectedBodySize = INDEX_NONE);
 
 		/**
 		 * Append an existing outfit facade to this collection.
@@ -129,6 +157,12 @@ namespace UE::Chaos::OutfitAsset
 		}
 
 	private:
+		void AddOutfit(
+			const FGuid& Guid,
+			const int32 BodySize,
+			const FString& Name,
+			const TMap<FGuid, FString>& Pieces);
+
 		Private::FOutfitCollection* GetOutfitCollection() { return const_cast<Private::FOutfitCollection*>(OutfitCollection.Get()); }
 	};
 }
diff --git a/Engine/Plugins/Experimental/ChaosOutfitAsset/Source/ChaosOutfitAssetEngine/Public/ChaosOutfitAsset/Outfit.h b/Engine/Plugins/Experimental/ChaosOutfitAsset/Source/ChaosOutfitAssetEngine/Public/ChaosOutfitAsset/Outfit.h
--- a/Engine/Plugins/Experimental/ChaosOutfitAsset/Source/ChaosOutfitAssetEngine/Public/ChaosOutfitAsset/Outfit.h
+++ b/Engine/Plugins/Experimental/ChaosOutfitAsset/Source/ChaosOutfitAssetEngine/Public/ChaosOutfitAsset/Outfit.h
@@ -10,11 +10,12 @@
 
 #include "Outfit.generated.h"
 
-class UChaosClothAssetBase;
 class FSkeletalMeshRenderData;
+class UChaosClothAssetBase;
 class UPhysicsAsset;
 class USkeletalMesh;
 struct FChaosClothSimulationModel;
+struct FChaosSizedOutfitSource;
 struct FManagedArrayCollection;
 struct FSkeletalMaterial;
 
@@ -46,6 +47,7 @@ struct CHAOSOUTFITASSETENGINE_API FChaosOutfitPiece final
 		const TArray<TSharedRef<const FManagedArrayCollection>>& InCollections);
 
 	FChaosOutfitPiece(const UChaosClothAssetBase& ClothAssetBase, int32 ModelIndex);
+	FChaosOutfitPiece(const TArray<FChaosOutfitPiece>& Pieces, int32 ModelIndex) : FChaosOutfitPiece(Pieces[ModelIndex]) {}
 
 	FChaosOutfitPiece(const FChaosOutfitPiece& Other);
 
@@ -86,8 +88,10 @@ public:
 	UChaosOutfit(FVTableHelper& Helper);  // This is declared so we can use TUniquePtr<FClothSimulationModel> with just a forward declare of that class
 	virtual ~UChaosOutfit() override;
 
-	void Append(const UChaosOutfit& Other);
+	void Append(const UChaosOutfit& Other, const FString& BodySizeNameFilter = {});
+
 	void Add(const UChaosClothAssetBase& ClothAssetBase);
+	void Add(const FChaosSizedOutfitSource& SizedOutfitSource, const FGuid& OutfitGuid = FGuid::NewGuid());
 
 	void CopyTo(
 		TArray<FChaosOutfitPiece>& OutPieces,
@@ -115,6 +119,8 @@ public:
 		return ReferenceSkeleton;
 	}
 
+	bool HasBodySize(const FString& SizeName) const;
+
 	//~ Begin UObject interface
 	virtual void Serialize(FArchive& Ar) override;
 	//~ End UObject interface
@@ -139,14 +145,15 @@ private:
 		const TArray<int32>& BoneMap,
 		OutLODRenderDataType& OutLODRenderData);
 
-	template<typename InRenderDataType, typename OutRenderDataType>
+	template<typename InRenderDataType, typename OutRenderDataType, typename PiecesType>
 	static void Merge(
 		const FReferenceSkeleton& InReferenceSkeleton,
 		const InRenderDataType* const InRenderData,
 		const TArray<FSkeletalMaterial>& InMaterials,
 		const FManagedArrayCollection& InOutfitCollection,
+		const PiecesType& Pieces,
 		const int32 NumPieces,
-		TFunction<FChaosOutfitPiece(int32)>&& MakePiece,
+		const FString& BodySizeNameFilter,
 		TArray<FChaosOutfitPiece>& OutPieces,
 		FReferenceSkeleton& OutReferenceSkeleton,
 		OutRenderDataType* const OutRenderData,
diff --git a/Engine/Plugins/Experimental/ChaosOutfitAsset/Source/ChaosOutfitAssetEngine/Public/ChaosOutfitAsset/OutfitCollection.h b/Engine/Plugins/Experimental/ChaosOutfitAsset/Source/ChaosOutfitAssetEngine/Public/ChaosOutfitAsset/OutfitCollection.h
--- a/Engine/Plugins/Experimental/ChaosOutfitAsset/Source/ChaosOutfitAssetEngine/Public/ChaosOutfitAsset/OutfitCollection.h
+++ b/Engine/Plugins/Experimental/ChaosOutfitAsset/Source/ChaosOutfitAssetEngine/Public/ChaosOutfitAsset/OutfitCollection.h
@@ -50,4 +50,7 @@ namespace UE::Chaos::OutfitAsset
 			}
 		}
 	}  // namespace OutfitCollection
+
+	/** Default body size name, when an outfit is created without any specific body size. */
+	CHAOSOUTFITASSETENGINE_API extern const FName DefaultBodySize;
 }  // namespace UE::Chaos::OutfitAsset
diff --git a/Engine/Plugins/Experimental/ChaosOutfitAsset/Source/ChaosOutfitAssetEngine/Public/ChaosOutfitAsset/SizedOutfitSource.h b/Engine/Plugins/Experimental/ChaosOutfitAsset/Source/ChaosOutfitAssetEngine/Public/ChaosOutfitAsset/SizedOutfitSource.h
new file mode 100644
--- /dev/null
+++ b/Engine/Plugins/Experimental/ChaosOutfitAsset/Source/ChaosOutfitAssetEngine/Public/ChaosOutfitAsset/SizedOutfitSource.h
@@ -0,0 +1,50 @@
+// Copyright Epic Games, Inc. All Rights Reserved.
+
+#pragma once 
+
+#include "Containers/Array.h"
+#include "UObject/ObjectPtr.h"
+#include "SizedOutfitSource.generated.h"
+
+class UChaosClothAssetBase;
+class USkeletalMesh;
+
+/**
+ * Input structure for setting up a single sized outfit.
+ */
+USTRUCT(BlueprintType)  // BlueprintType to allow this structure as a variable type.
+struct FChaosSizedOutfitSource final
+{
+	GENERATED_USTRUCT_BODY()
+
+	/**
+	 * The cloth or outfit asset to assign for this body size.
+	 * The asset must be the exact same garment representation for each specified sizes.
+	 * If multiple garments in multiple sizes are needed, they will have to be composed in another Outfit asset.
+	 */
+	UPROPERTY(EditAnywhere, Category = "OutfitSourceAsset")
+	TObjectPtr<const UChaosClothAssetBase> SourceAsset;
+
+	/**
+	 * The unique name of this body size.
+	 * The name of the first valid body part skeletal mesh will be used if empty.
+	 */
+	UPROPERTY(EditAnywhere, Category = "OutfitBodySize")
+	FString SizeName;
+
+	/**
+	 * The list of body part skeletal meshes making up the source body for this size name.
+	 * Usually one single MetaHuman merged body+head skeletal mesh.
+	 */
+	UPROPERTY(EditAnywhere, Category = "OutfitBodySize", Meta = (NoElementDuplicate))
+	TArray<TObjectPtr<const USkeletalMesh>> SourceBodyParts = { nullptr };
+
+	/**
+	 * Get a usable name for the body size.
+	 * Return SizeName unless empty, in which case return the name of the first valid body part.
+	 * If both SizeName and the SourceBodyParts are empty, but a valid SourceAsset is provided,
+	 * it will return "Default" as the body size name.
+	 * Otherwise if none of these are valid, an empty string is returned.
+	 */
+	CHAOSOUTFITASSETENGINE_API FString GetBodySizeName() const;
+};
