Chaos Outfit Asset - Added the GetOutfitClothCollections node.

#rnx

diff --git a/Engine/Plugins/Experimental/ChaosOutfitAsset/Source/ChaosOutfitAssetDataflowNodes/Private/ChaosOutfitAsset/ChaosOutfitAssetDataflowNodesModule.cpp b/Engine/Plugins/Experimental/ChaosOutfitAsset/Source/ChaosOutfitAssetDataflowNodes/Private/ChaosOutfitAsset/ChaosOutfitAssetDataflowNodesModule.cpp
--- a/Engine/Plugins/Experimental/ChaosOutfitAsset/Source/ChaosOutfitAssetDataflowNodes/Private/ChaosOutfitAsset/ChaosOutfitAssetDataflowNodesModule.cpp
+++ b/Engine/Plugins/Experimental/ChaosOutfitAsset/Source/ChaosOutfitAssetDataflowNodes/Private/ChaosOutfitAsset/ChaosOutfitAssetDataflowNodesModule.cpp
@@ -4,6 +4,7 @@
 #include "ChaosOutfitAsset/FilterSizedOutfitNode.h"
 #include "ChaosOutfitAsset/GetClothAssetNode.h"
 #include "ChaosOutfitAsset/GetOutfitAssetNode.h"
+#include "ChaosOutfitAsset/GetOutfitClothCollectionsNode.h"
 #include "ChaosOutfitAsset/MakeOutfitNode.h"
 #include "ChaosOutfitAsset/MakeSizedOutfitNode.h"
 #include "ChaosOutfitAsset/MergeOutfitsNode.h"
@@ -30,6 +31,7 @@ namespace UE::Chaos::OutfitAsset
 			DATAFLOW_NODE_REGISTER_CREATION_FACTORY_NODE_COLORS_BY_CATEGORY("Outfit", OutfitAssetNodeHeaderColor, OutfitAssetNodeBodyColor);
 			DATAFLOW_NODE_REGISTER_CREATION_FACTORY(FChaosGetClothAssetNode);
 			DATAFLOW_NODE_REGISTER_CREATION_FACTORY(FChaosGetOutfitAssetNode);
+			DATAFLOW_NODE_REGISTER_CREATION_FACTORY(FChaosGetOutfitClothCollectionsNode);
 			DATAFLOW_NODE_REGISTER_CREATION_FACTORY(FChaosOutfitAssetFilterSizedOutfitNode);
 			DATAFLOW_NODE_REGISTER_CREATION_FACTORY(FChaosOutfitAssetMakeOutfitNode);
 			DATAFLOW_NODE_REGISTER_CREATION_FACTORY(FChaosOutfitAssetMakeSizedOutfitNode);
diff --git a/Engine/Plugins/Experimental/ChaosOutfitAsset/Source/ChaosOutfitAssetDataflowNodes/Private/ChaosOutfitAsset/GetOutfitClothCollectionsNode.cpp b/Engine/Plugins/Experimental/ChaosOutfitAsset/Source/ChaosOutfitAssetDataflowNodes/Private/ChaosOutfitAsset/GetOutfitClothCollectionsNode.cpp
new file mode 100644
--- /dev/null
+++ b/Engine/Plugins/Experimental/ChaosOutfitAsset/Source/ChaosOutfitAssetDataflowNodes/Private/ChaosOutfitAsset/GetOutfitClothCollectionsNode.cpp
@@ -0,0 +1,52 @@
+// Copyright Epic Games, Inc. All Rights Reserved.
+
+#include "ChaosOutfitAsset/GetOutfitClothCollectionsNode.h"
+#include "ChaosOutfitAsset/Outfit.h"
+
+#include UE_INLINE_GENERATED_CPP_BY_NAME(GetOutfitClothCollectionsNode)
+
+#define LOCTEXT_NAMESPACE "ChaosGetOutfitClothCollectionsNode"
+
+FChaosGetOutfitClothCollectionsNode::FChaosGetOutfitClothCollectionsNode(const UE::Dataflow::FNodeParameters& InParam, FGuid InGuid)
+	: FDataflowNode(InParam, InGuid)
+{
+	RegisterInputConnection(&Outfit);
+	RegisterOutputConnection(&Outfit, &Outfit);
+	RegisterOutputConnection(&ClothCollections);
+	RegisterOutputConnection(&NumLods);
+	RegisterOutputConnection(&NumPieces);
+}
+
+void FChaosGetOutfitClothCollectionsNode::Evaluate(UE::Dataflow::FContext& Context, const FDataflowOutput* Out) const
+{
+	if (Out->IsA(&Outfit) || Out->IsA(&ClothCollections) || Out->IsA(&NumLods))
+	{
+		TObjectPtr<const UChaosOutfit> OutOutfit;
+		if (const TObjectPtr<const UChaosOutfit>& InOutfit = GetValue(Context, &Outfit))
+		{
+			OutOutfit = InOutfit;
+			SafeForwardInput(Context, &Outfit, &Outfit);
+		}
+		else
+		{
+			// No input Outfit, make up an empty output
+			OutOutfit = NewObject<UChaosOutfit>();
+			SetValue(Context, OutOutfit, &Outfit);
+		}
+
+		const int32 OutNumLods = (LodIndex == INDEX_NONE) ? OutOutfit->GetNumLods() : FMath::Min(1, OutOutfit->GetNumLods());
+		const int32 OutNumPieces = OutOutfit->GetPieces().Num();
+
+		TArray<FManagedArrayCollection> OutClothCollections;
+		OutClothCollections.Reserve(OutNumLods * OutNumPieces);
+		for (const TSharedRef<const FManagedArrayCollection>& ClothCollection : OutOutfit->GetClothCollections(LodIndex))
+		{
+			OutClothCollections.Emplace(*ClothCollection);
+		}
+		SetValue(Context, MoveTemp(OutClothCollections), &ClothCollections);
+		SetValue(Context, OutNumLods, &NumLods);
+		SetValue(Context, OutNumPieces, &NumPieces);
+	}
+}
+
+#undef LOCTEXT_NAMESPACE
diff --git a/Engine/Plugins/Experimental/ChaosOutfitAsset/Source/ChaosOutfitAssetDataflowNodes/Private/ChaosOutfitAsset/GetOutfitClothCollectionsNode.h b/Engine/Plugins/Experimental/ChaosOutfitAsset/Source/ChaosOutfitAssetDataflowNodes/Private/ChaosOutfitAsset/GetOutfitClothCollectionsNode.h
new file mode 100644
--- /dev/null
+++ b/Engine/Plugins/Experimental/ChaosOutfitAsset/Source/ChaosOutfitAssetDataflowNodes/Private/ChaosOutfitAsset/GetOutfitClothCollectionsNode.h
@@ -0,0 +1,47 @@
+// Copyright Epic Games, Inc. All Rights Reserved.
+
+#pragma once 
+
+#include "Dataflow/DataflowNode.h"
+#include "GeometryCollection/ManagedArrayCollection.h"
+#include "GetOutfitClothCollectionsNode.generated.h"
+
+class UChaosOutfit;
+
+/**
+ * Extract the cloth collections contained into the specified source outfit.
+ */
+USTRUCT(Meta = (Experimental, DataflowOutfit))
+struct FChaosGetOutfitClothCollectionsNode final : public FDataflowNode
+{
+	GENERATED_USTRUCT_BODY()
+	DATAFLOW_NODE_DEFINE_INTERNAL(FChaosGetOutfitClothCollectionsNode, "GetOutfitClothCollections", "Outfit", "Outfit Cloth Collections")
+
+public:
+	FChaosGetOutfitClothCollectionsNode(const UE::Dataflow::FNodeParameters& InParam, FGuid InGuid = FGuid::NewGuid());
+
+private:
+	//~ Begin FDataflowNode interface
+	virtual void Evaluate(UE::Dataflow::FContext& Context, const FDataflowOutput* Out) const override;
+	//~ End FDataflowNode interface
+
+	/** The source outfit. */
+	UPROPERTY(Meta = (DataflowInput, DataflowOutput, DataflowPassthrough = "Outfit"))
+	TObjectPtr<const UChaosOutfit> Outfit;
+
+	/** The outfit cloth collections array in group of consecutives LODs for each of the outfit pieces. */
+	UPROPERTY(Meta = (DataflowOutput))
+	TArray<FManagedArrayCollection> ClothCollections;
+
+	/** The LOD to output in the cloth collections array. Set to -1 to output all LODs */
+	UPROPERTY(EditAnywhere, Category = "Get Outfit Cloth Collections", DisplayName = "LOD Index", Meta = (UIMin = -1))
+	int32 LodIndex = -1;
+
+	/** The number of LODs output in the cloth collections array. */
+	UPROPERTY(DisplayName = "Num LODs", Meta = (DataflowOutput))
+	int32 NumLods = 0;
+
+	/** The number of cloth pieces output in the cloth collections array. */
+	UPROPERTY(Meta = (DataflowOutput))
+	int32 NumPieces = 0;
+};
diff --git a/Engine/Plugins/Experimental/ChaosOutfitAsset/Source/ChaosOutfitAssetEngine/Private/ChaosOutfitAsset/Outfit.cpp b/Engine/Plugins/Experimental/ChaosOutfitAsset/Source/ChaosOutfitAssetEngine/Private/ChaosOutfitAsset/Outfit.cpp
--- a/Engine/Plugins/Experimental/ChaosOutfitAsset/Source/ChaosOutfitAssetEngine/Private/ChaosOutfitAsset/Outfit.cpp
+++ b/Engine/Plugins/Experimental/ChaosOutfitAsset/Source/ChaosOutfitAssetEngine/Private/ChaosOutfitAsset/Outfit.cpp
@@ -4,6 +4,7 @@
 #include "Chaos/CollectionPropertyFacade.h"
 #include "ChaosClothAsset/ClothAsset.h"
 #include "ChaosClothAsset/ClothSimulationModel.h"
+#include "ChaosClothAsset/CollectionClothFacade.h"
 #include "ChaosOutfitAsset/BodyUserData.h"
 #include "ChaosOutfitAsset/CollectionOutfitFacade.h"
 #include "ChaosOutfitAsset/OutfitAsset.h"
@@ -780,6 +781,49 @@ void UChaosOutfit::CopyTo(
 		OutOutfitCollection);
 }
 
+int32 UChaosOutfit::GetNumLods() const
+{
+	int32 NumLods = 0;
+	for (const FChaosOutfitPiece& Piece : Pieces)
+	{
+		NumLods = FMath::Max(Piece.Collections.Num(), NumLods);
+	}
+	return NumLods;
+}
+
+TArray<TSharedRef<const FManagedArrayCollection>> UChaosOutfit::GetClothCollections(int32 StartLodIndex) const
+{
+	int32 NumLods;
+	if (StartLodIndex == INDEX_NONE)
+	{
+		// Get all LODs
+		NumLods = GetNumLods();
+		StartLodIndex = 0;
+	}
+	else
+	{
+		NumLods = 1;
+	}
+
+	TArray<TSharedRef<const FManagedArrayCollection>> ClothCollections;
+	ClothCollections.Reserve(NumLods * Pieces.Num());
+
+	TSharedRef<FManagedArrayCollection> EmptyClothCollection = MakeShared<FManagedArrayCollection>();
+	UE::Chaos::ClothAsset::FCollectionClothFacade ClothFacade(EmptyClothCollection);
+	ClothFacade.DefineSchema();
+
+	for (const FChaosOutfitPiece& Piece : Pieces)
+	{
+		for (int32 LodIndex = StartLodIndex; LodIndex < StartLodIndex + NumLods; ++LodIndex)
+		{
+			ClothCollections.Emplace(Piece.Collections.IsValidIndex(LodIndex) ?
+				Piece.Collections[LodIndex] :
+				EmptyClothCollection);
+		}
+	}
+	return ClothCollections;
+}
+
 bool UChaosOutfit::HasBodySize(const FString& SizeName) const
 {
 	UE::Chaos::OutfitAsset::FCollectionOutfitConstFacade OutfitFacade(OutfitCollection);
diff --git a/Engine/Plugins/Experimental/ChaosOutfitAsset/Source/ChaosOutfitAssetEngine/Public/ChaosOutfitAsset/Outfit.h b/Engine/Plugins/Experimental/ChaosOutfitAsset/Source/ChaosOutfitAssetEngine/Public/ChaosOutfitAsset/Outfit.h
--- a/Engine/Plugins/Experimental/ChaosOutfitAsset/Source/ChaosOutfitAssetEngine/Public/ChaosOutfitAsset/Outfit.h
+++ b/Engine/Plugins/Experimental/ChaosOutfitAsset/Source/ChaosOutfitAssetEngine/Public/ChaosOutfitAsset/Outfit.h
@@ -114,6 +114,12 @@ public:
 		return OutfitCollection;
 	}
 
+	/** Return the number of LODs (max LOD contained in any piece). */
+	int32 GetNumLods() const;
+
+	/** Return the outfit pieces cloth collections for all LODs. */
+	TArray<TSharedRef<const FManagedArrayCollection>> GetClothCollections(int32 LodIndex) const;
+
 	const FReferenceSkeleton& GetReferenceSkeleton() const
 	{
 		return ReferenceSkeleton;
