Chaos Outfit Resizing: new MakeResizableOutfit and ResizeOutfitTemplate dataflow graphs and required code changes to support them.

The MakeResizableOutfit generates an OutfitAsset which is ready for resizing--it caches the size and interpolation data in the asset.

The ResizeOutfitTemplate takes a ResizableOutfit and a TargetBody and refits the outfit to fit the TargetBody. There are a few other variables to control graph execution including explicitly selecting a BodySize, and skipping Sim Mesh and Skin Weight generation.

Lots of new nodes and changes to existing nodes to support this workflow. Of note:

* Added new MeshResizingCore library (and thus MeshResizing plugin) dependency in order to calculate and store interpolation data in the asset.
* Added interpolation data to the Outfit Collection schema.
* Added generation of interpolation data when adding an outfit size.
* Updated RBFInterpolation (core and nodes) to work with MeshDescriptions, in order to avoid making intermediate DynamicMeshes
* Added SetOutfitClothCollectionNode which can overwrite a cloth collection in an outfit (used to ensure proxy deformer data is stored in the cloth collection)
* Flipped behavior of FilterSizedOutfitNode to use BodySize input first and TargetBody second

#virtualized

diff --git a/Engine/Plugins/ChaosClothAssetEditor/Source/ChaosClothAssetDataflowNodes/Private/ChaosClothAsset/MakeClothAssetNode.cpp b/Engine/Plugins/ChaosClothAssetEditor/Source/ChaosClothAssetDataflowNodes/Private/ChaosClothAsset/MakeClothAssetNode.cpp
--- a/Engine/Plugins/ChaosClothAssetEditor/Source/ChaosClothAssetDataflowNodes/Private/ChaosClothAsset/MakeClothAssetNode.cpp
+++ b/Engine/Plugins/ChaosClothAssetEditor/Source/ChaosClothAssetDataflowNodes/Private/ChaosClothAsset/MakeClothAssetNode.cpp
@@ -20,6 +20,7 @@ FChaosClothAssetMakeClothAssetNode::FChaosClothAssetMakeClothAssetNode(const UE:
 	: FDataflowNode(InParam, InGuid)
 {
 	RegisterOutputConnection(&ClothAsset);
+	RegisterInputConnection(&CollectionLodsArray);
 
 	for (int32 Index = 0; Index < NumInitialCollectionLods; ++Index)
 	{
@@ -90,43 +91,75 @@ TArray<TSharedRef<const FManagedArrayCollection>> FChaosClothAssetMakeClothAsset
 	using namespace UE::Chaos::ClothAsset;
 
 	TArray<TSharedRef<const FManagedArrayCollection>> CollectionLodValues;
-	CollectionLodValues.Reserve(CollectionLods.Num());
-
 	int32 LastValidLodIndex = INDEX_NONE;
-	for (int32 LodIndex = 0; LodIndex < CollectionLods.Num(); ++LodIndex)
+	auto AddCleanedCollectionLod = [this, &LastValidLodIndex, &CollectionLodValues](int32 LodIndex, TSharedRef<FManagedArrayCollection> CollectionLodValue)
+		{
+			FCollectionClothFacade ClothFacade(CollectionLodValue);
+			if (ClothFacade.HasValidRenderData())  // The cloth collection must at least have a render mesh
+			{
+				FClothGeometryTools::CleanupAndCompactMesh(CollectionLodValue);
+				LastValidLodIndex = LodIndex;
+			}
+			else if (LastValidLodIndex >= 0)
+			{
+				ClothFacade.DefineSchema();
+				ClothFacade.Initialize(FCollectionClothConstFacade(CollectionLodValues[LastValidLodIndex]));
+
+				FClothDataflowTools::LogAndToastWarning(*this,
+					LOCTEXT("InvalidInputLodNHeadline", "Invalid input LOD."),
+					FText::Format(
+						LOCTEXT("InvalidInputLodNDetails",
+							"Invalid or empty input LOD for LOD {0}.\n"
+							"Using the previous valid LOD {1} instead."),
+						LodIndex,
+						LastValidLodIndex));
+			}
+			else
+			{
+				FClothDataflowTools::LogAndToastWarning(*this,
+					LOCTEXT("InvalidInputLod0Headline", "Invalid input LOD 0."),
+					LOCTEXT("InvalidInputLod0Details",
+						"Invalid or empty input LOD for LOD 0.\n"
+						"LOD 0 cannot be empty in order to construct a valid Cloth Asset."));
+				return;
+			}
+			CollectionLodValues.Emplace(MoveTemp(CollectionLodValue));
+		};
+
+
+	if (IsConnected(&CollectionLodsArray))
 	{
-		TSharedRef<FManagedArrayCollection> CollectionLodValue = MakeShared<FManagedArrayCollection>(GetValue<FManagedArrayCollection>(Context, GetConnectionReference(LodIndex)));
+		bool bAlsoHasElementConnections = false;
+		for (int32 LodIndex = 0; LodIndex < CollectionLods.Num(); ++LodIndex)
+		{
+			if (IsConnected(GetConnectionReference(LodIndex)))
+			{
+				bAlsoHasElementConnections = true;
+				break;
+			}
+		}
 
-		FCollectionClothFacade ClothFacade(CollectionLodValue);
-		if (ClothFacade.HasValidRenderData())  // The cloth collection must at least have a render mesh
+		if (bAlsoHasElementConnections)
 		{
-			FClothGeometryTools::CleanupAndCompactMesh(CollectionLodValue);
-			LastValidLodIndex = LodIndex;
+			Context.Warning(TEXT("Connections found both to CollectionsLodsArray and individual CollectionLods. Only CollectionsLodArray will be used"), this);
 		}
-		else if (LastValidLodIndex >= 0)
+
+		const TArray<FManagedArrayCollection>& InCollectionLodsArray = GetValue(Context, &CollectionLodsArray);
+
+		CollectionLodValues.Reserve(InCollectionLodsArray.Num());
+
+		for (int32 LodIndex = 0; LodIndex < InCollectionLodsArray.Num(); ++LodIndex)
 		{
-			ClothFacade.DefineSchema();
-			ClothFacade.Initialize(FCollectionClothConstFacade(CollectionLodValues[LastValidLodIndex]));
-
-			FClothDataflowTools::LogAndToastWarning(*this,
-				LOCTEXT("InvalidInputLodNHeadline", "Invalid input LOD."),
-				FText::Format(
-					LOCTEXT("InvalidInputLodNDetails",
-						"Invalid or empty input LOD for LOD {0}.\n"
-						"Using the previous valid LOD {1} instead."),
-					LodIndex,
-					LastValidLodIndex));
+			AddCleanedCollectionLod(LodIndex, MakeShared<FManagedArrayCollection>(InCollectionLodsArray[LodIndex]));
 		}
-		else
+	}
+	else
+	{
+		CollectionLodValues.Reserve(CollectionLods.Num());
+		for (int32 LodIndex = 0; LodIndex < CollectionLods.Num(); ++LodIndex)
 		{
-			FClothDataflowTools::LogAndToastWarning(*this,
-				LOCTEXT("InvalidInputLod0Headline", "Invalid input LOD 0."),
-				LOCTEXT("InvalidInputLod0Details",
-					"Invalid or empty input LOD for LOD 0.\n"
-					"LOD 0 cannot be empty in order to construct a valid Cloth Asset."));
-			break;
+			AddCleanedCollectionLod(LodIndex, MakeShared<FManagedArrayCollection>(GetValue<FManagedArrayCollection>(Context, GetConnectionReference(LodIndex))));			
 		}
-		CollectionLodValues.Emplace(MoveTemp(CollectionLodValue));
 	}
 	return CollectionLodValues;
 }
diff --git a/Engine/Plugins/ChaosClothAssetEditor/Source/ChaosClothAssetDataflowNodes/Private/ChaosClothAsset/MakeClothAssetNode.h b/Engine/Plugins/ChaosClothAssetEditor/Source/ChaosClothAssetDataflowNodes/Private/ChaosClothAsset/MakeClothAssetNode.h
--- a/Engine/Plugins/ChaosClothAssetEditor/Source/ChaosClothAssetDataflowNodes/Private/ChaosClothAsset/MakeClothAssetNode.h
+++ b/Engine/Plugins/ChaosClothAssetEditor/Source/ChaosClothAssetDataflowNodes/Private/ChaosClothAsset/MakeClothAssetNode.h
@@ -33,7 +33,11 @@ private:
 	TArray<TSharedRef<const FManagedArrayCollection>> GetCleanedCollectionLodValues(UE::Dataflow::FContext& Context) const;
 	UE::Dataflow::TConnectionReference<FManagedArrayCollection> GetConnectionReference(int32 Index) const;
 
-	/** Input cloth collection for this LOD. */
+	/** Input cloth collections for this LOD -- Array connection. Individual CollectionLods will be ignored if there is a CollectionLodsArray connection. */
+	UPROPERTY(Meta = (DataflowInput))
+	TArray<FManagedArrayCollection> CollectionLodsArray;
+
+	/** Input cloth collection for this LOD. Individual CollectionLods will be ignored if there is a CollectionLodsArray connection. */
 	UPROPERTY()
 	TArray<FManagedArrayCollection> CollectionLods;
 
@@ -41,6 +45,6 @@ private:
 	UPROPERTY(Meta = (DataflowOutput))
 	TObjectPtr<const UChaosClothAsset> ClothAsset;
 
-	static constexpr int32 NumRequiredInputs = 0;
+	static constexpr int32 NumRequiredInputs = 1;
 	static constexpr int32 NumInitialCollectionLods = 1;
 };
diff --git a/Engine/Plugins/Experimental/ChaosOutfitAsset/ChaosOutfitAsset.uplugin b/Engine/Plugins/Experimental/ChaosOutfitAsset/ChaosOutfitAsset.uplugin
--- a/Engine/Plugins/Experimental/ChaosOutfitAsset/ChaosOutfitAsset.uplugin
+++ b/Engine/Plugins/Experimental/ChaosOutfitAsset/ChaosOutfitAsset.uplugin
@@ -41,6 +41,10 @@
 		{
 			"Name": "Dataflow",
 			"Enabled": true
+		},
+		{
+			"Name": "MeshResizing",
+			"Enabled": true
 		}
 	]
 }
diff --git a/Engine/Plugins/Experimental/ChaosOutfitAsset/Source/ChaosOutfitAssetDataflowNodes/ChaosOutfitAssetDataflowNodes.Build.cs b/Engine/Plugins/Experimental/ChaosOutfitAsset/Source/ChaosOutfitAssetDataflowNodes/ChaosOutfitAssetDataflowNodes.Build.cs
--- a/Engine/Plugins/Experimental/ChaosOutfitAsset/Source/ChaosOutfitAssetDataflowNodes/ChaosOutfitAssetDataflowNodes.Build.cs
+++ b/Engine/Plugins/Experimental/ChaosOutfitAsset/Source/ChaosOutfitAssetDataflowNodes/ChaosOutfitAssetDataflowNodes.Build.cs
@@ -18,7 +18,8 @@ public class ChaosOutfitAssetDataflowNodes : ModuleRules
 				"ChaosOutfitAssetEngine",
 				"DataflowCore",
 				"DataflowNodes",
-				"Engine"
+				"Engine",
+				"MeshResizingCore",
 			}
 		);
 	}
diff --git a/Engine/Plugins/Experimental/ChaosOutfitAsset/Source/ChaosOutfitAssetDataflowNodes/Private/ChaosOutfitAsset/ChaosOutfitAssetDataflowNodesModule.cpp b/Engine/Plugins/Experimental/ChaosOutfitAsset/Source/ChaosOutfitAssetDataflowNodes/Private/ChaosOutfitAsset/ChaosOutfitAssetDataflowNodesModule.cpp
--- a/Engine/Plugins/Experimental/ChaosOutfitAsset/Source/ChaosOutfitAssetDataflowNodes/Private/ChaosOutfitAsset/ChaosOutfitAssetDataflowNodesModule.cpp
+++ b/Engine/Plugins/Experimental/ChaosOutfitAsset/Source/ChaosOutfitAssetDataflowNodes/Private/ChaosOutfitAsset/ChaosOutfitAssetDataflowNodesModule.cpp
@@ -4,11 +4,14 @@
 #include "ChaosOutfitAsset/FilterSizedOutfitNode.h"
 #include "ChaosOutfitAsset/GetClothAssetNode.h"
 #include "ChaosOutfitAsset/GetOutfitAssetNode.h"
+#include "ChaosOutfitAsset/GetOutfitBodyPartsNode.h"
 #include "ChaosOutfitAsset/GetOutfitClothCollectionsNode.h"
+#include "ChaosOutfitAsset/GetOutfitRBFInterpolationDataNode.h"
 #include "ChaosOutfitAsset/MakeOutfitNode.h"
 #include "ChaosOutfitAsset/MakeSizedOutfitNode.h"
 #include "ChaosOutfitAsset/MergeOutfitsNode.h"
 #include "ChaosOutfitAsset/OutfitAssetTerminalNode.h"
+#include "ChaosOutfitAsset/SetOutfitClothCollectionNode.h"
 #include "ChaosOutfitAsset/SizedOutfitSourceAnyType.h"
 #include "Dataflow/DataflowAnyTypeRegistry.h"
 #include "Dataflow/DataflowNode.h"
@@ -29,13 +32,17 @@ namespace UE::Chaos::OutfitAsset
 		static void RegisterDataflowNodes()
 		{
 			DATAFLOW_NODE_REGISTER_CREATION_FACTORY_NODE_COLORS_BY_CATEGORY("Outfit", OutfitAssetNodeHeaderColor, OutfitAssetNodeBodyColor);
+			DATAFLOW_NODE_REGISTER_CREATION_FACTORY(FChaosExtractBodyPartsArrayFromBodySizePartsNode);
 			DATAFLOW_NODE_REGISTER_CREATION_FACTORY(FChaosGetClothAssetNode);
 			DATAFLOW_NODE_REGISTER_CREATION_FACTORY(FChaosGetOutfitAssetNode);
+			DATAFLOW_NODE_REGISTER_CREATION_FACTORY(FChaosGetOutfitBodyPartsNode);
 			DATAFLOW_NODE_REGISTER_CREATION_FACTORY(FChaosGetOutfitClothCollectionsNode);
+			DATAFLOW_NODE_REGISTER_CREATION_FACTORY(FChaosGetOutfitRBFInterpolationDataNode);
 			DATAFLOW_NODE_REGISTER_CREATION_FACTORY(FChaosOutfitAssetFilterSizedOutfitNode);
 			DATAFLOW_NODE_REGISTER_CREATION_FACTORY(FChaosOutfitAssetMakeOutfitNode);
 			DATAFLOW_NODE_REGISTER_CREATION_FACTORY(FChaosOutfitAssetMakeSizedOutfitNode);
 			DATAFLOW_NODE_REGISTER_CREATION_FACTORY(FChaosOutfitAssetMergeOutfitsNode);
+			DATAFLOW_NODE_REGISTER_CREATION_FACTORY(FChaosSetOutfitClothCollectionNode);
 			DATAFLOW_NODE_REGISTER_CREATION_FACTORY(FChaosOutfitAssetTerminalNode);
 PRAGMA_DISABLE_DEPRECATION_WARNINGS
 			// Move any deprecated node registrations under here
diff --git a/Engine/Plugins/Experimental/ChaosOutfitAsset/Source/ChaosOutfitAssetDataflowNodes/Private/ChaosOutfitAsset/FilterSizedOutfitNode.cpp b/Engine/Plugins/Experimental/ChaosOutfitAsset/Source/ChaosOutfitAssetDataflowNodes/Private/ChaosOutfitAsset/FilterSizedOutfitNode.cpp
--- a/Engine/Plugins/Experimental/ChaosOutfitAsset/Source/ChaosOutfitAssetDataflowNodes/Private/ChaosOutfitAsset/FilterSizedOutfitNode.cpp
+++ b/Engine/Plugins/Experimental/ChaosOutfitAsset/Source/ChaosOutfitAssetDataflowNodes/Private/ChaosOutfitAsset/FilterSizedOutfitNode.cpp
@@ -31,23 +31,21 @@ void FChaosOutfitAssetFilterSizedOutfitNode::Evaluate(UE::Dataflow::FContext& Co
 			FCollectionOutfitConstFacade InOutfitFacade(InOutfit->GetOutfitCollection());
 			if (InOutfitFacade.IsValid())
 			{
-				FString InSizeName;
-
-				if (const TObjectPtr<const USkeletalMesh>& InTargetBody = GetValue(Context, &TargetBody))
+				FString InSizeName = GetValue(Context, &SizeName);
+				if (!InSizeName.IsEmpty() && !InOutfit->HasBodySize(InSizeName))
 				{
-					// Selection by matching target body
-					const int32 ClosestBodySize = InOutfitFacade.FindClosestBodySize(*InTargetBody);
-					InSizeName = InOutfitFacade.GetBodySizeName(ClosestBodySize);  // Empty string when ClosestBodySize == INDEX_NONE
+					// Only show a warning when the size name is non empty
+					Context.Warning(FString::Printf(TEXT("The given body size [%s] doesn't exist in the input Outfit."), *InSizeName), this, Out);
+					InSizeName.Reset();
 				}
-				else
+
+				if (InSizeName.IsEmpty())
 				{
-					// Selection by name
-					InSizeName = GetValue(Context, &SizeName);
-					if (!InSizeName.IsEmpty() && !InOutfit->HasBodySize(InSizeName))
+					if (const TObjectPtr<const USkeletalMesh>& InTargetBody = GetValue(Context, &TargetBody))
 					{
-						// Only show a warning when the size name is non empty
-						Context.Warning(FString::Printf(TEXT("The given body size [%s] doesn't exist in the input Outfit."), *SizeName), this, Out);
-						InSizeName.Reset();
+						// Selection by matching target body
+						const int32 ClosestBodySize = InOutfitFacade.FindClosestBodySize(*InTargetBody);
+						InSizeName = InOutfitFacade.GetBodySizeName(ClosestBodySize);  // Empty string when ClosestBodySize == INDEX_NONE
 					}
 				}
 
diff --git a/Engine/Plugins/Experimental/ChaosOutfitAsset/Source/ChaosOutfitAssetDataflowNodes/Private/ChaosOutfitAsset/FilterSizedOutfitNode.h b/Engine/Plugins/Experimental/ChaosOutfitAsset/Source/ChaosOutfitAssetDataflowNodes/Private/ChaosOutfitAsset/FilterSizedOutfitNode.h
--- a/Engine/Plugins/Experimental/ChaosOutfitAsset/Source/ChaosOutfitAssetDataflowNodes/Private/ChaosOutfitAsset/FilterSizedOutfitNode.h
+++ b/Engine/Plugins/Experimental/ChaosOutfitAsset/Source/ChaosOutfitAssetDataflowNodes/Private/ChaosOutfitAsset/FilterSizedOutfitNode.h
@@ -36,15 +36,14 @@ private:
 
 	/**
 	 * The name of the body size to use to filter.
-	 * If the input Size Name is empty, the output will be set to a name set from the body part measurements (if provided).
-	 * Requires the Target Body to be empty for the input Size Name to have an effect.
+	 * If the input Size Name is empty, the output will be set to a name set from the Target Body's measurements.
 	 */
 	UPROPERTY(EditAnywhere, Category = "FilterSizedOutfit", Meta = (DataflowInput, DataflowOutput, DataflowPassthrough = "SizeName", EditCondition = "TargetBody == nullptr"))
 	FString SizeName;
 
 	/**
 	 * The target body skeletal mesh containing the measurements to select the required size to use to filter.
-	 * The input Size Name is ignored when a valid target body is specified, and the selected size is returned as the output Size Name.
+	 * The target body is unused if Size Name is a valid name.
 	 */
 	UPROPERTY(EditAnywhere, Category = "FilterSizedOutfit", Meta = (DataflowInput))
 	TObjectPtr<const USkeletalMesh> TargetBody;
diff --git a/Engine/Plugins/Experimental/ChaosOutfitAsset/Source/ChaosOutfitAssetDataflowNodes/Private/ChaosOutfitAsset/GetOutfitBodyPartsNode.cpp b/Engine/Plugins/Experimental/ChaosOutfitAsset/Source/ChaosOutfitAssetDataflowNodes/Private/ChaosOutfitAsset/GetOutfitBodyPartsNode.cpp
new file mode 100644
--- /dev/null
+++ b/Engine/Plugins/Experimental/ChaosOutfitAsset/Source/ChaosOutfitAssetDataflowNodes/Private/ChaosOutfitAsset/GetOutfitBodyPartsNode.cpp
@@ -0,0 +1,76 @@
+// Copyright Epic Games, Inc. All Rights Reserved.
+
+#include "ChaosOutfitAsset/GetOutfitBodyPartsNode.h"
+#include "ChaosOutfitAsset/Outfit.h"
+#include "ChaosOutfitAsset/CollectionOutfitFacade.h"
+#include "Engine/SkeletalMesh.h"
+
+#include UE_INLINE_GENERATED_CPP_BY_NAME(GetOutfitBodyPartsNode)
+
+#define LOCTEXT_NAMESPACE "ChaosGetOutfitBodyPartsNode"
+
+FChaosGetOutfitBodyPartsNode::FChaosGetOutfitBodyPartsNode(const UE::Dataflow::FNodeParameters& InParam, FGuid InGuid)
+	: FDataflowNode(InParam, InGuid)
+{
+	RegisterInputConnection(&Outfit);
+	RegisterOutputConnection(&Outfit, &Outfit);
+	RegisterOutputConnection(&BodySizeParts);
+}
+
+void FChaosGetOutfitBodyPartsNode::Evaluate(UE::Dataflow::FContext& Context, const FDataflowOutput* Out) const
+{
+	if (Out->IsA(&Outfit) || Out->IsA(&BodySizeParts))
+	{
+		TObjectPtr<const UChaosOutfit> OutOutfit;
+		if (const TObjectPtr<const UChaosOutfit>& InOutfit = GetValue(Context, &Outfit))
+		{
+			OutOutfit = InOutfit;
+			SafeForwardInput(Context, &Outfit, &Outfit);
+		}
+		else
+		{
+			// No input Outfit, make up an empty output
+			OutOutfit = NewObject<UChaosOutfit>();
+			SetValue(Context, OutOutfit, &Outfit);
+		}
+		using namespace UE::Chaos::OutfitAsset;
+
+		const FCollectionOutfitConstFacade OutfitFacade(OutOutfit->GetOutfitCollection());
+		const int32 NumBodySizes = OutfitFacade.GetNumBodySizes();
+
+		TArray<FChaosOutfitBodySizeBodyParts> OutBodyParts;
+		OutBodyParts.SetNum(NumBodySizes);
+		for (int32 BodySizeIndex = 0; BodySizeIndex < NumBodySizes; ++BodySizeIndex)
+		{
+			const TConstArrayView<FString> BodyPartNames = OutfitFacade.GetBodySizeBodyPartsSkeletalMeshes(BodySizeIndex);
+			OutBodyParts[BodySizeIndex].BodyParts.Reserve(BodyPartNames.Num());
+
+			for (const FString& BodyPartName : BodyPartNames)
+			{
+				if (const TObjectPtr<const USkeletalMesh> SKM = LoadObject<USkeletalMesh>(nullptr, *BodyPartName))
+				{
+					OutBodyParts[BodySizeIndex].BodyParts.Add(SKM);
+				}
+			}
+		}
+		SetValue(Context, MoveTemp(OutBodyParts), &BodySizeParts);
+	}
+}
+
+FChaosExtractBodyPartsArrayFromBodySizePartsNode::FChaosExtractBodyPartsArrayFromBodySizePartsNode(const UE::Dataflow::FNodeParameters& InParam, FGuid InGuid)
+	: FDataflowNode(InParam, InGuid)
+{
+	RegisterInputConnection(&BodySizeParts);
+	RegisterOutputConnection(&BodyParts);
+}
+
+void FChaosExtractBodyPartsArrayFromBodySizePartsNode::Evaluate(UE::Dataflow::FContext& Context, const FDataflowOutput* Out) const
+{
+	if (Out->IsA(&BodyParts))
+	{
+		const FChaosOutfitBodySizeBodyParts& InBodySizeParts = GetValue(Context, &BodySizeParts);
+		SetValue(Context, InBodySizeParts.BodyParts, &BodyParts);
+	}
+}
+
+#undef LOCTEXT_NAMESPACE
diff --git a/Engine/Plugins/Experimental/ChaosOutfitAsset/Source/ChaosOutfitAssetDataflowNodes/Private/ChaosOutfitAsset/GetOutfitBodyPartsNode.h b/Engine/Plugins/Experimental/ChaosOutfitAsset/Source/ChaosOutfitAssetDataflowNodes/Private/ChaosOutfitAsset/GetOutfitBodyPartsNode.h
new file mode 100644
--- /dev/null
+++ b/Engine/Plugins/Experimental/ChaosOutfitAsset/Source/ChaosOutfitAssetDataflowNodes/Private/ChaosOutfitAsset/GetOutfitBodyPartsNode.h
@@ -0,0 +1,68 @@
+// Copyright Epic Games, Inc. All Rights Reserved.
+
+#pragma once 
+
+#include "Dataflow/DataflowNode.h"
+#include "GeometryCollection/ManagedArrayCollection.h"
+#include "GetOutfitBodyPartsNode.generated.h"
+
+class UChaosOutfit;
+class USkeletalMesh;
+
+USTRUCT()
+struct FChaosOutfitBodySizeBodyParts
+{
+	GENERATED_USTRUCT_BODY()
+	UPROPERTY()
+	TArray<TObjectPtr<const USkeletalMesh>> BodyParts;
+};
+
+/**
+ * Extract the Body Part Skeletal Meshes from an Outfit.
+ */
+USTRUCT(Meta = (Experimental, DataflowOutfit))
+struct FChaosGetOutfitBodyPartsNode final : public FDataflowNode
+{
+	GENERATED_USTRUCT_BODY()
+	DATAFLOW_NODE_DEFINE_INTERNAL(FChaosGetOutfitBodyPartsNode, "GetOutfitBodyParts", "Outfit", "Outfit Body Parts Skeletal Mesh")
+
+public:
+	FChaosGetOutfitBodyPartsNode(const UE::Dataflow::FNodeParameters& InParam, FGuid InGuid = FGuid::NewGuid());
+
+private:
+	//~ Begin FDataflowNode interface
+	virtual void Evaluate(UE::Dataflow::FContext& Context, const FDataflowOutput* Out) const override;
+	//~ End FDataflowNode interface
+
+	/** The source outfit. */
+	UPROPERTY(Meta = (DataflowInput, DataflowOutput, DataflowPassthrough = "Outfit"))
+	TObjectPtr<const UChaosOutfit> Outfit;
+
+	/** The outfit body parts grouped by BodySize. */
+	UPROPERTY(Meta = (DataflowOutput))
+	TArray<FChaosOutfitBodySizeBodyParts> BodySizeParts;
+};
+
+/** Extract the array of BodyParts from a FChaosOutfitBodySizeBodyParts */
+USTRUCT(Meta = (Experimental, DataflowOutfit))
+struct FChaosExtractBodyPartsArrayFromBodySizePartsNode final : public FDataflowNode
+{
+	GENERATED_USTRUCT_BODY()
+	DATAFLOW_NODE_DEFINE_INTERNAL(FChaosExtractBodyPartsArrayFromBodySizePartsNode, "ExtractBodyPartsArrayFromBodySizeParts", "Outfit", "Extract Outfit Body Parts Skeletal Mesh")
+
+public:
+	FChaosExtractBodyPartsArrayFromBodySizePartsNode(const UE::Dataflow::FNodeParameters& InParam, FGuid InGuid = FGuid::NewGuid());
+
+private:
+	//~ Begin FDataflowNode interface
+	virtual void Evaluate(UE::Dataflow::FContext& Context, const FDataflowOutput* Out) const override;
+	//~ End FDataflowNode interface
+
+	/** The source outfit. */
+	UPROPERTY(Meta = (DataflowInput))
+	FChaosOutfitBodySizeBodyParts BodySizeParts;
+
+	/** The outfit body parts grouped by BodySize. */
+	UPROPERTY(Meta = (DataflowOutput))
+	TArray<TObjectPtr<const USkeletalMesh>> BodyParts;
+};
diff --git a/Engine/Plugins/Experimental/ChaosOutfitAsset/Source/ChaosOutfitAssetDataflowNodes/Private/ChaosOutfitAsset/GetOutfitRBFInterpolationDataNode.cpp b/Engine/Plugins/Experimental/ChaosOutfitAsset/Source/ChaosOutfitAssetDataflowNodes/Private/ChaosOutfitAsset/GetOutfitRBFInterpolationDataNode.cpp
new file mode 100644
--- /dev/null
+++ b/Engine/Plugins/Experimental/ChaosOutfitAsset/Source/ChaosOutfitAssetDataflowNodes/Private/ChaosOutfitAsset/GetOutfitRBFInterpolationDataNode.cpp
@@ -0,0 +1,59 @@
+// Copyright Epic Games, Inc. All Rights Reserved.
+
+#include "ChaosOutfitAsset/GetOutfitRBFInterpolationDataNode.h"
+#include "ChaosOutfitAsset/Outfit.h"
+#include "ChaosOutfitAsset/CollectionOutfitFacade.h"
+
+#include UE_INLINE_GENERATED_CPP_BY_NAME(GetOutfitRBFInterpolationDataNode)
+
+#define LOCTEXT_NAMESPACE "ChaosGetOutfitRBFInterpolationDataNode"
+
+FChaosGetOutfitRBFInterpolationDataNode::FChaosGetOutfitRBFInterpolationDataNode(const UE::Dataflow::FNodeParameters& InParam, FGuid InGuid)
+	: FDataflowNode(InParam, InGuid)
+{
+	RegisterInputConnection(&Outfit);
+	RegisterInputConnection(&BodySizeIndex);
+	RegisterInputConnection(&BodyPartIndex);
+	RegisterOutputConnection(&Outfit, &Outfit);
+	RegisterOutputConnection(&InterpolationData);
+}
+
+void FChaosGetOutfitRBFInterpolationDataNode::Evaluate(UE::Dataflow::FContext& Context, const FDataflowOutput* Out) const
+{
+	if (Out->IsA(&Outfit) || Out->IsA(&InterpolationData))
+	{
+		FMeshResizingRBFInterpolationData OutInterpolationData;
+		if (const TObjectPtr<const UChaosOutfit>& InOutfit = GetValue(Context, &Outfit))
+		{
+			using namespace UE::Chaos::OutfitAsset;
+			const FCollectionOutfitConstFacade OutfitFacade(InOutfit->GetOutfitCollection());
+
+			const int32 InBodySizeIndex = GetValue(Context, &BodySizeIndex);
+			const int32 InBodyPartIndex = GetValue(Context, &BodyPartIndex);
+			if (InBodySizeIndex >= 0 && InBodySizeIndex < OutfitFacade.GetNumBodySizes())
+			{
+				const FCollectionOutfitConstFacade::FRBFInterpolationDataWrapper InterpolationDataWrapper = OutfitFacade.GetBodySizeInterpolationData(InBodySizeIndex);
+				if (InterpolationDataWrapper.SampleIndices.IsValidIndex(InBodyPartIndex))
+				{
+					OutInterpolationData.SampleIndices = InterpolationDataWrapper.SampleIndices[InBodyPartIndex];
+					OutInterpolationData.SampleRestPositions = InterpolationDataWrapper.SampleRestPositions[InBodyPartIndex];
+					OutInterpolationData.InterpolationWeights = InterpolationDataWrapper.InterpolationWeights[InBodyPartIndex];
+				}
+				else
+				{
+					Context.Warning(FString::Printf(TEXT("The given BodyPartIndex [%d] doesn't index into this BodySize's NumBodyParts [%d]"), InBodyPartIndex, InterpolationDataWrapper.SampleIndices.Num()), this, Out);
+				}
+			}
+			else
+			{
+				Context.Warning(FString::Printf(TEXT("The given BodySizeIndex [%d] doesn't index into the NumBodySizes [%d]"), InBodySizeIndex, OutfitFacade.GetNumBodySizes()), this, Out);
+			}
+
+		}
+		SafeForwardInput(Context, &Outfit, &Outfit);
+		SetValue(Context, MoveTemp(OutInterpolationData), &InterpolationData);
+	}
+}
+
+
+#undef LOCTEXT_NAMESPACE
diff --git a/Engine/Plugins/Experimental/ChaosOutfitAsset/Source/ChaosOutfitAssetDataflowNodes/Private/ChaosOutfitAsset/GetOutfitRBFInterpolationDataNode.h b/Engine/Plugins/Experimental/ChaosOutfitAsset/Source/ChaosOutfitAssetDataflowNodes/Private/ChaosOutfitAsset/GetOutfitRBFInterpolationDataNode.h
new file mode 100644
--- /dev/null
+++ b/Engine/Plugins/Experimental/ChaosOutfitAsset/Source/ChaosOutfitAssetDataflowNodes/Private/ChaosOutfitAsset/GetOutfitRBFInterpolationDataNode.h
@@ -0,0 +1,42 @@
+// Copyright Epic Games, Inc. All Rights Reserved.
+
+#pragma once 
+
+#include "Dataflow/DataflowNode.h"
+#include "MeshResizing/RBFInterpolation.h"
+#include "GetOutfitRBFInterpolationDataNode.generated.h"
+
+class UChaosOutfit;
+/**
+ * Extract the Body Part RBF Interpolation Data from an outfit.
+ */
+USTRUCT(Meta = (Experimental, DataflowOutfit))
+struct FChaosGetOutfitRBFInterpolationDataNode final : public FDataflowNode
+{
+	GENERATED_USTRUCT_BODY()
+	DATAFLOW_NODE_DEFINE_INTERNAL(FChaosGetOutfitRBFInterpolationDataNode, "GetOutfitRBFInterpolationData", "Outfit", "Outfit RBF Interpolation Data")
+
+public:
+	FChaosGetOutfitRBFInterpolationDataNode(const UE::Dataflow::FNodeParameters& InParam, FGuid InGuid = FGuid::NewGuid());
+
+private:
+	//~ Begin FDataflowNode interface
+	virtual void Evaluate(UE::Dataflow::FContext& Context, const FDataflowOutput* Out) const override;
+	//~ End FDataflowNode interface
+
+	/** The source outfit. */
+	UPROPERTY(Meta = (DataflowInput, DataflowOutput, DataflowPassthrough = "Outfit"))
+	TObjectPtr<const UChaosOutfit> Outfit;
+
+	/** The body size index.*/
+	UPROPERTY(EditAnywhere, Category = "Outfit RBF Interpolation", Meta = (DataflowInput, ClampMin = 0))
+	int32 BodySizeIndex = 0;
+
+	/** The body part (sub) index.*/
+	UPROPERTY(EditAnywhere, Category = "Outfit RBF Interpolation", Meta = (DataflowInput, ClampMin = 0))
+	int32 BodyPartIndex = 0;
+
+	/** The interpolation data. */
+	UPROPERTY(Meta = (DataflowOutput))
+	FMeshResizingRBFInterpolationData InterpolationData;
+};
\ No newline at end of file
diff --git a/Engine/Plugins/Experimental/ChaosOutfitAsset/Source/ChaosOutfitAssetDataflowNodes/Private/ChaosOutfitAsset/SetOutfitClothCollectionNode.cpp b/Engine/Plugins/Experimental/ChaosOutfitAsset/Source/ChaosOutfitAssetDataflowNodes/Private/ChaosOutfitAsset/SetOutfitClothCollectionNode.cpp
new file mode 100644
--- /dev/null
+++ b/Engine/Plugins/Experimental/ChaosOutfitAsset/Source/ChaosOutfitAssetDataflowNodes/Private/ChaosOutfitAsset/SetOutfitClothCollectionNode.cpp
@@ -0,0 +1,57 @@
+// Copyright Epic Games, Inc. All Rights Reserved.
+
+#include "ChaosOutfitAsset/SetOutfitClothCollectionNode.h"
+#include "ChaosOutfitAsset/Outfit.h"
+
+#include UE_INLINE_GENERATED_CPP_BY_NAME(SetOutfitClothCollectionNode)
+
+#define LOCTEXT_NAMESPACE "ChaosGetOutfitClothCollectionsNode"
+
+FChaosSetOutfitClothCollectionNode::FChaosSetOutfitClothCollectionNode(const UE::Dataflow::FNodeParameters& InParam, FGuid InGuid)
+	: FDataflowNode(InParam, InGuid)
+{
+	RegisterInputConnection(&Outfit);
+	RegisterInputConnection(&ClothCollection);
+	RegisterInputConnection(&PieceIndex);
+	RegisterInputConnection(&LODIndex);
+	RegisterOutputConnection(&Outfit, &Outfit);
+}
+
+void FChaosSetOutfitClothCollectionNode::Evaluate(UE::Dataflow::FContext& Context, const FDataflowOutput* Out) const
+{
+	if (Out->IsA(&Outfit))
+	{
+		if (const TObjectPtr<const UChaosOutfit>& InOutfit = GetValue(Context, &Outfit))
+		{
+			const int32 InPieceIndex = GetValue(Context, &PieceIndex);
+			const int32 InLODIndex = GetValue(Context, &LODIndex);
+			if (InOutfit->GetPieces().IsValidIndex(InPieceIndex))
+			{
+
+				if (InOutfit->GetPieces()[InPieceIndex].Collections.IsValidIndex(InLODIndex))
+
+				{
+					const FManagedArrayCollection& InCollection = GetValue<FManagedArrayCollection>(Context, &ClothCollection);
+
+					TObjectPtr<UChaosOutfit> OutOutfit = NewObject<UChaosOutfit>();
+					OutOutfit->Append(*InOutfit);
+
+					OutOutfit->GetPieces()[InPieceIndex].Collections[InLODIndex] = MakeShared<const FManagedArrayCollection>(InCollection);
+					SetValue<TObjectPtr<const UChaosOutfit>>(Context, OutOutfit, &Outfit);
+					return;
+				}
+				else
+				{
+					Context.Warning(FString::Printf(TEXT("The given LODIndex [%d] doesn't index into this Piece's LOD array of size [%d]"), InLODIndex, InOutfit->GetPieces()[InPieceIndex].Collections.Num()), this, Out);
+				}
+			}
+			else
+			{
+				Context.Warning(FString::Printf(TEXT("The given PieceIndex [%d] doesn't index into Pieces array of size [%d]"), InPieceIndex, InOutfit->GetPieces().Num()), this, Out);
+			}
+		}
+		SafeForwardInput(Context, &Outfit, &Outfit);
+	}
+}
+
+#undef LOCTEXT_NAMESPACE
diff --git a/Engine/Plugins/Experimental/ChaosOutfitAsset/Source/ChaosOutfitAssetDataflowNodes/Private/ChaosOutfitAsset/SetOutfitClothCollectionNode.h b/Engine/Plugins/Experimental/ChaosOutfitAsset/Source/ChaosOutfitAssetDataflowNodes/Private/ChaosOutfitAsset/SetOutfitClothCollectionNode.h
new file mode 100644
--- /dev/null
+++ b/Engine/Plugins/Experimental/ChaosOutfitAsset/Source/ChaosOutfitAssetDataflowNodes/Private/ChaosOutfitAsset/SetOutfitClothCollectionNode.h
@@ -0,0 +1,44 @@
+// Copyright Epic Games, Inc. All Rights Reserved.
+
+#pragma once 
+
+#include "Dataflow/DataflowNode.h"
+#include "GeometryCollection/ManagedArrayCollection.h"
+#include "SetOutfitClothCollectionNode.generated.h"
+
+class UChaosOutfit;
+
+/**
+ * Replace the ClothCollection in an Outfit with a new one.
+ * Any data derived from the ClothCollection (e.g., Simulation Model, Render Data) will NOT be regenerated in the Outfit.
+ */
+USTRUCT(Meta = (Experimental, DataflowOutfit))
+struct FChaosSetOutfitClothCollectionNode final : public FDataflowNode
+{
+	GENERATED_USTRUCT_BODY()
+	DATAFLOW_NODE_DEFINE_INTERNAL(FChaosSetOutfitClothCollectionNode, "SetOutfitClothCollection", "Outfit", "Outfit Cloth Collections")
+
+public:
+	FChaosSetOutfitClothCollectionNode(const UE::Dataflow::FNodeParameters& InParam, FGuid InGuid = FGuid::NewGuid());
+
+private:
+	//~ Begin FDataflowNode interface
+	virtual void Evaluate(UE::Dataflow::FContext& Context, const FDataflowOutput* Out) const override;
+	//~ End FDataflowNode interface
+
+	/** The outfit to be edited. */
+	UPROPERTY(Meta = (DataflowInput, DataflowOutput, DataflowPassthrough = "Outfit"))
+	TObjectPtr<const UChaosOutfit> Outfit;
+
+	/** The replacement cloth collection. */
+	UPROPERTY(Meta = (DataflowInput))
+	FManagedArrayCollection ClothCollection;
+
+	/** The Outfit Piece to replace. */
+	UPROPERTY(EditAnywhere, Category = "Set Cloth Collection", Meta =(ClampMin = 0, DataflowInput))
+	int32 PieceIndex = 0;
+
+	/** The Outfit LOD to replace. */
+	UPROPERTY(EditAnywhere, Category = "Set Cloth Collection", Meta = (ClampMin = 0, DataflowInput))
+	int32 LODIndex = 0;
+};
diff --git a/Engine/Plugins/Experimental/ChaosOutfitAsset/Source/ChaosOutfitAssetEngine/ChaosOutfitAssetEngine.Build.cs b/Engine/Plugins/Experimental/ChaosOutfitAsset/Source/ChaosOutfitAssetEngine/ChaosOutfitAssetEngine.Build.cs
--- a/Engine/Plugins/Experimental/ChaosOutfitAsset/Source/ChaosOutfitAssetEngine/ChaosOutfitAssetEngine.Build.cs
+++ b/Engine/Plugins/Experimental/ChaosOutfitAsset/Source/ChaosOutfitAssetEngine/ChaosOutfitAssetEngine.Build.cs
@@ -26,6 +26,7 @@ public class ChaosOutfitAssetEngine : ModuleRules
 				"ChaosClothAssetEngine",
 				"ClothingSystemRuntimeCommon",
 				"Engine",
+				"MeshResizingCore",
 				"RenderCore",
 				"RHI",
 			}
diff --git a/Engine/Plugins/Experimental/ChaosOutfitAsset/Source/ChaosOutfitAssetEngine/Private/ChaosOutfitAsset/CollectionOutfitFacade.cpp b/Engine/Plugins/Experimental/ChaosOutfitAsset/Source/ChaosOutfitAssetEngine/Private/ChaosOutfitAsset/CollectionOutfitFacade.cpp
--- a/Engine/Plugins/Experimental/ChaosOutfitAsset/Source/ChaosOutfitAssetEngine/Private/ChaosOutfitAsset/CollectionOutfitFacade.cpp
+++ b/Engine/Plugins/Experimental/ChaosOutfitAsset/Source/ChaosOutfitAssetEngine/Private/ChaosOutfitAsset/CollectionOutfitFacade.cpp
@@ -89,6 +89,9 @@ namespace UE::Chaos::OutfitAsset
 			struct
 			{
 				TManagedArrayPtr<FString> SkeletalMesh;  // Body parts skeletal mesh (several per sizes, e.g. body, head)
+				TManagedArrayPtr<TArray<int32>> RBFInterpolationSampleIndices; // Precomputed RBF interpolation sample indices
+				TManagedArrayPtr<TArray<FVector3f>> RBFInterpolationSampleRestPositions; // Precomputed RBF interpolation sample rest positions
+				TManagedArrayPtr<TArray<float>> RBFInterpolationWeights; // Precomputed RBF interpolation weights
 			} BodyParts;
 
 			struct
@@ -111,6 +114,9 @@ namespace UE::Chaos::OutfitAsset
 				BodySizes.BodyPartsStart = ManagedArrayCollection.FindAttribute<int32>(Attribute::BodySizes::BodyPartsStart, Group::BodySizes);
 				BodySizes.BodyPartsCount = ManagedArrayCollection.FindAttribute<int32>(Attribute::BodySizes::BodyPartsCount, Group::BodySizes);
 				BodyParts.SkeletalMesh = ManagedArrayCollection.FindAttribute<FString>(Attribute::BodyParts::SkeletalMesh, Group::BodyParts);
+				BodyParts.RBFInterpolationSampleIndices = ManagedArrayCollection.FindAttribute<TArray<int32>>(Attribute::BodyParts::RBFInterpolationSampleIndices, Group::BodyParts);
+				BodyParts.RBFInterpolationSampleRestPositions = ManagedArrayCollection.FindAttribute<TArray<FVector3f>>(Attribute::BodyParts::RBFInterpolationSampleRestPositions, Group::BodyParts);
+				BodyParts.RBFInterpolationWeights = ManagedArrayCollection.FindAttribute<TArray<float>>(Attribute::BodyParts::RBFInterpolationWeights, Group::BodyParts);
 				Measurements.Name = ManagedArrayCollection.FindAttribute<FString>(Attribute::Measurements::Name, Group::Measurements);
 			}
 
@@ -129,6 +135,9 @@ namespace UE::Chaos::OutfitAsset
 					BodySizes.BodyPartsStart &&
 					BodySizes.BodyPartsCount &&
 					BodyParts.SkeletalMesh &&
+					BodyParts.RBFInterpolationSampleIndices &&
+					BodyParts.RBFInterpolationSampleRestPositions &&
+					BodyParts.RBFInterpolationWeights &&
 					Measurements.Name;
 			}
 
@@ -148,6 +157,9 @@ namespace UE::Chaos::OutfitAsset
 					BodySizes.BodyPartsStart = &ManagedArrayCollection.AddAttribute<int32>(Attribute::BodySizes::BodyPartsStart, Group::BodySizes, Group::BodyParts);
 					BodySizes.BodyPartsCount = &ManagedArrayCollection.AddAttribute<int32>(Attribute::BodySizes::BodyPartsCount, Group::BodySizes);
 					BodyParts.SkeletalMesh = &ManagedArrayCollection.AddAttribute<FString>(Attribute::BodyParts::SkeletalMesh, Group::BodyParts);
+					BodyParts.RBFInterpolationSampleIndices = &ManagedArrayCollection.AddAttribute<TArray<int32>>(Attribute::BodyParts::RBFInterpolationSampleIndices, Group::BodyParts);
+					BodyParts.RBFInterpolationSampleRestPositions = &ManagedArrayCollection.AddAttribute<TArray<FVector3f>>(Attribute::BodyParts::RBFInterpolationSampleRestPositions, Group::BodyParts);
+					BodyParts.RBFInterpolationWeights = &ManagedArrayCollection.AddAttribute<TArray<float>>(Attribute::BodyParts::RBFInterpolationWeights, Group::BodyParts);
 					Measurements.Name = &ManagedArrayCollection.AddAttribute<FString>(Attribute::Measurements::Name, Group::Measurements);
 				}
 			}
@@ -545,6 +557,46 @@ namespace UE::Chaos::OutfitAsset
 		}
 		return TConstArrayView<FString>();
 	}
+	
+	int32 FCollectionOutfitConstFacade::GetBodySizeBodyPartOffset(const int32 BodySize) const
+	{
+		check(IsValid());
+		if (GetOutfitCollection()->BodySizes.BodyPartsStart->IsValidIndex(BodySize))
+		{
+			return (*GetOutfitCollection()->BodySizes.BodyPartsStart)[BodySize];
+		}
+		return INDEX_NONE;
+	}
+
+	int32 FCollectionOutfitConstFacade::GetBodySizeBodyPartCount(const int32 BodySize) const
+	{
+		check(IsValid());
+		if (GetOutfitCollection()->BodySizes.BodyPartsStart->IsValidIndex(BodySize))
+		{
+			return (*GetOutfitCollection()->BodySizes.BodyPartsCount)[BodySize];
+		}
+		return 0;
+	}
+
+	FCollectionOutfitConstFacade::FRBFInterpolationDataWrapper FCollectionOutfitConstFacade::GetBodySizeInterpolationData(const int32 BodySize) const
+	{
+		check(IsValid());
+		FRBFInterpolationDataWrapper Result;
+		if (GetOutfitCollection()->BodySizes.BodyPartsStart->IsValidIndex(BodySize))
+		{
+			const int32 BodyPartsStart = (*GetOutfitCollection()->BodySizes.BodyPartsStart)[BodySize];
+			const int32 BodyPartsCount = (*GetOutfitCollection()->BodySizes.BodyPartsCount)[BodySize];
+			if (BodyPartsCount &&
+				GetOutfitCollection()->BodyParts.RBFInterpolationSampleIndices->IsValidIndex(BodyPartsStart) &&
+				GetOutfitCollection()->BodyParts.RBFInterpolationSampleIndices->IsValidIndex(BodyPartsStart + BodyPartsCount - 1))
+			{
+				Result.SampleIndices = TConstArrayView<TArray<int32>>(GetOutfitCollection()->BodyParts.RBFInterpolationSampleIndices->GetData() + BodyPartsStart, BodyPartsCount);
+				Result.SampleRestPositions = TConstArrayView<TArray<FVector3f>>(GetOutfitCollection()->BodyParts.RBFInterpolationSampleRestPositions->GetData() + BodyPartsStart, BodyPartsCount);
+				Result.InterpolationWeights = TConstArrayView<TArray<float>>(GetOutfitCollection()->BodyParts.RBFInterpolationWeights->GetData() + BodyPartsStart, BodyPartsCount);
+			}
+		}
+		return Result;
+	}
 
 	TMap<FString, float> FCollectionOutfitConstFacade::GetBodySizeMeasurements(const int32 BodySize) const
 	{
@@ -591,7 +643,7 @@ namespace UE::Chaos::OutfitAsset
 		check(IsValid());
 	}
 
-	int32 FCollectionOutfitFacade::AddBodySize(const FString& Name, const TConstArrayView<FString> BodyPartsSkeletalMeshes, const TMap<FString, float>& Measurements)
+	int32 FCollectionOutfitFacade::AddBodySize(const FString& Name, const TConstArrayView<FString> BodyPartsSkeletalMeshes, const TMap<FString, float>& Measurements, const FRBFInterpolationDataWrapper& InterpolationData)
 	{
 		using namespace OutfitCollection;
 		check(IsValid());
@@ -614,9 +666,20 @@ namespace UE::Chaos::OutfitAsset
 		if (BodyPartsCount)
 		{
 			BodyPartsStart = GetOutfitCollection()->ManagedArrayCollection.AddElements(BodyPartsCount, Group::BodyParts);
+			const bool bInterpolationDataValid =
+				InterpolationData.SampleIndices.Num() == BodyPartsCount &&
+				InterpolationData.SampleRestPositions.Num() == BodyPartsCount &&
+				InterpolationData.InterpolationWeights.Num() == BodyPartsCount;
+				
 			for (int32 Index = 0; Index < BodyPartsCount; ++Index)
 			{
 				(*GetOutfitCollection()->BodyParts.SkeletalMesh)[BodyPartsStart + Index] = BodyPartsSkeletalMeshes[Index];
+				if (bInterpolationDataValid)
+				{
+					(*GetOutfitCollection()->BodyParts.RBFInterpolationSampleIndices)[BodyPartsStart + Index] = InterpolationData.SampleIndices[Index];
+					(*GetOutfitCollection()->BodyParts.RBFInterpolationSampleRestPositions)[BodyPartsStart + Index] = InterpolationData.SampleRestPositions[Index];
+					(*GetOutfitCollection()->BodyParts.RBFInterpolationWeights)[BodyPartsStart + Index] = InterpolationData.InterpolationWeights[Index];
+				}
 			}
 		}
 		else
@@ -757,7 +820,8 @@ namespace UE::Chaos::OutfitAsset
 			RemappedBodySize[BodySize - SelectedBodySize] = AddBodySize(
 				Other.GetBodySizeName(BodySize),
 				Other.GetBodySizeBodyPartsSkeletalMeshes(BodySize),
-				Other.GetBodySizeMeasurements(BodySize));
+				Other.GetBodySizeMeasurements(BodySize),
+				Other.GetBodySizeInterpolationData(BodySize));
 		}
 
 		// Merge sized outfits
diff --git a/Engine/Plugins/Experimental/ChaosOutfitAsset/Source/ChaosOutfitAssetEngine/Private/ChaosOutfitAsset/Outfit.cpp b/Engine/Plugins/Experimental/ChaosOutfitAsset/Source/ChaosOutfitAssetEngine/Private/ChaosOutfitAsset/Outfit.cpp
--- a/Engine/Plugins/Experimental/ChaosOutfitAsset/Source/ChaosOutfitAssetEngine/Private/ChaosOutfitAsset/Outfit.cpp
+++ b/Engine/Plugins/Experimental/ChaosOutfitAsset/Source/ChaosOutfitAssetEngine/Private/ChaosOutfitAsset/Outfit.cpp
@@ -12,6 +12,7 @@
 #include "ChaosOutfitAsset/OutfitCollection.h"
 #include "ChaosOutfitAsset/SizedOutfitSource.h"
 #include "Engine/SkeletalMesh.h"
+#include "MeshResizing/RBFInterpolation.h"
 #include "Rendering/SkeletalMeshLODModel.h"
 #include "Rendering/SkeletalMeshLODRenderData.h"
 #include "Rendering/SkeletalMeshRenderData.h"
@@ -695,6 +696,9 @@ void UChaosOutfit::Add(const FChaosSizedOutfitSource& SizedOutfitSource, const F
 	if (!BodySizeName.IsEmpty())
 	{
 		TArray<FString> BodyPartsSkeletalMeshes;
+		TArray<TArray<int32>> InterpolationData_SampleIndices;
+		TArray<TArray<FVector3f>> InterpolationData_SampleRestPositions;
+		TArray<TArray<float>> InterpolationData_InterpolationWeights;
 		BodyPartsSkeletalMeshes.Reserve(SizedOutfitSource.SourceBodyParts.Num());
 		TMap<FString, float> Measurements;
 		for (const TObjectPtr<const USkeletalMesh>& SourceBodyPart : SizedOutfitSource.SourceBodyParts)
@@ -714,9 +718,26 @@ void UChaosOutfit::Add(const FChaosSizedOutfitSource& SizedOutfitSource, const F
 					}
 				}
 				BodyPartsSkeletalMeshes.Emplace(SourceBodyPart->GetPathName());
+#if WITH_EDITORONLY_DATA
+				if (SizedOutfitSource.NumResizingInterpolationPoints > 0)
+				{
+					FMeshResizingRBFInterpolationData InterpData;
+					if (const FMeshDescription* const MeshDescription = SourceBodyPart->GetMeshDescription(0))
+					{
+						UE::MeshResizing::FRBFInterpolation::GenerateWeights(*MeshDescription, SizedOutfitSource.NumResizingInterpolationPoints, InterpData);
+					}
+					InterpolationData_SampleIndices.Emplace(MoveTemp(InterpData.SampleIndices));
+					InterpolationData_SampleRestPositions.Emplace(MoveTemp(InterpData.SampleRestPositions));
+					InterpolationData_InterpolationWeights.Emplace(MoveTemp(InterpData.InterpolationWeights));
+				}
+#endif
 			}
 		}
-		BodySize = InOutfitFacade.AddBodySize(BodySizeName, BodyPartsSkeletalMeshes, Measurements);
+		FCollectionOutfitConstFacade::FRBFInterpolationDataWrapper InterpolationData;
+		InterpolationData.SampleIndices = TConstArrayView<TArray<int32>>(InterpolationData_SampleIndices);
+		InterpolationData.SampleRestPositions = TConstArrayView<TArray<FVector3f>>(InterpolationData_SampleRestPositions);
+		InterpolationData.InterpolationWeights = TConstArrayView<TArray<float>>(InterpolationData_InterpolationWeights);
+		BodySize = InOutfitFacade.AddBodySize(BodySizeName, BodyPartsSkeletalMeshes, Measurements, InterpolationData);
 	}
 	else
 	{
diff --git a/Engine/Plugins/Experimental/ChaosOutfitAsset/Source/ChaosOutfitAssetEngine/Private/ChaosOutfitAsset/OutfitCollection.cpp b/Engine/Plugins/Experimental/ChaosOutfitAsset/Source/ChaosOutfitAssetEngine/Private/ChaosOutfitAsset/OutfitCollection.cpp
--- a/Engine/Plugins/Experimental/ChaosOutfitAsset/Source/ChaosOutfitAssetEngine/Private/ChaosOutfitAsset/OutfitCollection.cpp
+++ b/Engine/Plugins/Experimental/ChaosOutfitAsset/Source/ChaosOutfitAssetEngine/Private/ChaosOutfitAsset/OutfitCollection.cpp
@@ -47,6 +47,9 @@ namespace UE::Chaos::OutfitAsset
 			namespace BodyParts
 			{
 				const FName SkeletalMesh(TEXT("SkeletalMesh"));  // FString
+				const FName RBFInterpolationSampleIndices(TEXT("RBFInterpolationSampleIndices"));  // TArray<int32>
+				const FName RBFInterpolationSampleRestPositions(TEXT("RBFInterpolationSampleRestPositions"));  // TArray<FVector3f>
+				const FName RBFInterpolationWeights(TEXT("RBFInterpolationWeights"));  // // TArray<float>
 			}
 			// Measurements group
 			namespace Measurements
diff --git a/Engine/Plugins/Experimental/ChaosOutfitAsset/Source/ChaosOutfitAssetEngine/Public/ChaosOutfitAsset/CollectionOutfitFacade.h b/Engine/Plugins/Experimental/ChaosOutfitAsset/Source/ChaosOutfitAssetEngine/Public/ChaosOutfitAsset/CollectionOutfitFacade.h
--- a/Engine/Plugins/Experimental/ChaosOutfitAsset/Source/ChaosOutfitAssetEngine/Public/ChaosOutfitAsset/CollectionOutfitFacade.h
+++ b/Engine/Plugins/Experimental/ChaosOutfitAsset/Source/ChaosOutfitAssetEngine/Public/ChaosOutfitAsset/CollectionOutfitFacade.h
@@ -4,11 +4,13 @@
 
 #include "Templates/SharedPointer.h"
 #include "Templates/UniquePtr.h"
+#include "ChaosOutfitAsset/OutfitCollection.h"
 
 template<typename T> class TManagedArray;
 class UChaosClothAssetBase;
 class USkeletalMesh;
 struct FManagedArrayCollection;
+struct FMeshResizingRBFInterpolationData;
 
 namespace UE::Chaos::OutfitAsset
 {
@@ -24,6 +26,13 @@ namespace UE::Chaos::OutfitAsset
 	class FCollectionOutfitConstFacade
 	{
 	public:
+		struct CHAOSOUTFITASSETENGINE_API FRBFInterpolationDataWrapper
+		{
+			TConstArrayView<TArray<int32>> SampleIndices;
+			TConstArrayView<TArray<FVector3f>> SampleRestPositions;
+			TConstArrayView<TArray<float>> InterpolationWeights;
+		};
+
 		CHAOSOUTFITASSETENGINE_API explicit FCollectionOutfitConstFacade(const TSharedRef<const FManagedArrayCollection>& InManagedArrayCollection);
 		CHAOSOUTFITASSETENGINE_API explicit FCollectionOutfitConstFacade(const FManagedArrayCollection& InManagedArrayCollection);
 
@@ -79,6 +88,15 @@ namespace UE::Chaos::OutfitAsset
 		/** Return the body part skeletal meshes path name for the specified body size. */
 		CHAOSOUTFITASSETENGINE_API TConstArrayView<FString> GetBodySizeBodyPartsSkeletalMeshes(const int32 BodySize) const;
 
+		/** Return the body size body part offset (to index into the global body part array). Returns INDEX_NONE if BodySize is invalid.*/
+		CHAOSOUTFITASSETENGINE_API int32 GetBodySizeBodyPartOffset(const int32 BodySize) const;
+
+		/** Return the body size body part count. Returns 0 if BodySize is invalid.*/
+		CHAOSOUTFITASSETENGINE_API int32 GetBodySizeBodyPartCount(const int32 BodySize) const;
+
+		/** Return  for the specified body size. */
+		CHAOSOUTFITASSETENGINE_API FRBFInterpolationDataWrapper GetBodySizeInterpolationData(const int32 BodySize) const;
+
 		/** Return the body measurements stored for this specified body size. */
 		CHAOSOUTFITASSETENGINE_API TMap<FString, float> GetBodySizeMeasurements(const int32 BodySize) const;
 
@@ -120,13 +138,15 @@ namespace UE::Chaos::OutfitAsset
 		CHAOSOUTFITASSETENGINE_API int32 AddBodySize(
 			const FString& Name,
 			const TConstArrayView<FString> BodyPartsSkeletalMeshes,
-			const TMap<FString, float>& Measurements);
+			const TMap<FString, float>& Measurements,
+			const FRBFInterpolationDataWrapper& InterpolationData);
 
 		/** Find or add a new body size and returns its index. */
 		int32 FindOrAddBodySize(const FString& Name)
 		{
 			const int32 BodySize = FindBodySize(Name);
-			return BodySize != INDEX_NONE ? BodySize : AddBodySize(Name, {}, {});
+
+			return BodySize != INDEX_NONE ? BodySize : AddBodySize(Name, {}, {}, {});
 		}
 
 		/**
diff --git a/Engine/Plugins/Experimental/ChaosOutfitAsset/Source/ChaosOutfitAssetEngine/Public/ChaosOutfitAsset/Outfit.h b/Engine/Plugins/Experimental/ChaosOutfitAsset/Source/ChaosOutfitAssetEngine/Public/ChaosOutfitAsset/Outfit.h
--- a/Engine/Plugins/Experimental/ChaosOutfitAsset/Source/ChaosOutfitAssetEngine/Public/ChaosOutfitAsset/Outfit.h
+++ b/Engine/Plugins/Experimental/ChaosOutfitAsset/Source/ChaosOutfitAssetEngine/Public/ChaosOutfitAsset/Outfit.h
@@ -105,6 +105,11 @@ public:
 		return Pieces;
 	}
 
+	TArray<FChaosOutfitPiece>& GetPieces()
+	{
+		return Pieces;
+	}
+
 	const FManagedArrayCollection& GetOutfitCollection() const
 	{
 		return OutfitCollection;
diff --git a/Engine/Plugins/Experimental/ChaosOutfitAsset/Source/ChaosOutfitAssetEngine/Public/ChaosOutfitAsset/OutfitCollection.h b/Engine/Plugins/Experimental/ChaosOutfitAsset/Source/ChaosOutfitAssetEngine/Public/ChaosOutfitAsset/OutfitCollection.h
--- a/Engine/Plugins/Experimental/ChaosOutfitAsset/Source/ChaosOutfitAssetEngine/Public/ChaosOutfitAsset/OutfitCollection.h
+++ b/Engine/Plugins/Experimental/ChaosOutfitAsset/Source/ChaosOutfitAssetEngine/Public/ChaosOutfitAsset/OutfitCollection.h
@@ -43,6 +43,9 @@ namespace UE::Chaos::OutfitAsset
 			namespace BodyParts
 			{
 				CHAOSOUTFITASSETENGINE_API extern const FName SkeletalMesh;  // FString
+				CHAOSOUTFITASSETENGINE_API extern const FName RBFInterpolationSampleIndices;  // TArray<int32>
+				CHAOSOUTFITASSETENGINE_API extern const FName RBFInterpolationSampleRestPositions;  // TArray<FVector3f>
+				CHAOSOUTFITASSETENGINE_API extern const FName RBFInterpolationWeights;  // TArray<float>
 			}
 			namespace Measurements
 			{
@@ -53,4 +56,7 @@ namespace UE::Chaos::OutfitAsset
 
 	/** Default body size name, when an outfit is created without any specific body size. */
 	CHAOSOUTFITASSETENGINE_API extern const FName DefaultBodySize;
+
+	/** Default number of RBF Interpolation points used for RBF Resizing */
+	inline constexpr int32 DefaultNumRBFInterpolationPoints = 1500;
 }  // namespace UE::Chaos::OutfitAsset
diff --git a/Engine/Plugins/Experimental/ChaosOutfitAsset/Source/ChaosOutfitAssetEngine/Public/ChaosOutfitAsset/SizedOutfitSource.h b/Engine/Plugins/Experimental/ChaosOutfitAsset/Source/ChaosOutfitAssetEngine/Public/ChaosOutfitAsset/SizedOutfitSource.h
--- a/Engine/Plugins/Experimental/ChaosOutfitAsset/Source/ChaosOutfitAssetEngine/Public/ChaosOutfitAsset/SizedOutfitSource.h
+++ b/Engine/Plugins/Experimental/ChaosOutfitAsset/Source/ChaosOutfitAssetEngine/Public/ChaosOutfitAsset/SizedOutfitSource.h
@@ -4,6 +4,7 @@
 
 #include "Containers/Array.h"
 #include "UObject/ObjectPtr.h"
+#include "ChaosOutfitAsset/OutfitCollection.h"
 #include "SizedOutfitSource.generated.h"
 
 class UChaosClothAssetBase;
@@ -39,6 +40,15 @@ struct FChaosSizedOutfitSource final
 	UPROPERTY(EditAnywhere, Category = "OutfitBodySize", Meta = (NoElementDuplicate))
 	TArray<TObjectPtr<const USkeletalMesh>> SourceBodyParts = { nullptr };
 
+	/**
+	 * The number of interpolation points used in the resizing algorithm. These points are distributed evenly across the entire body.
+	 * Increasing this number increases the quality of the resizing operation, but at additional cost, including the initial operation of 
+	 *  generating your resizable outfit, and the size of your Outfit on disk.
+	 * If you find you cannot generate acceptable resizing results by increasing this number, we recommend generating a new Size.
+	 */
+	UPROPERTY(EditAnywhere, AdvancedDisplay, Category = "OutfitBodySize", Meta = (ClampMin = 0, ClampMax = 5000))
+	int32 NumResizingInterpolationPoints = UE::Chaos::OutfitAsset::DefaultNumRBFInterpolationPoints;
+
 	/**
 	 * Get a usable name for the body size.
 	 * Return SizeName unless empty, in which case return the name of the first valid body part.
diff --git a/Engine/Plugins/Experimental/MeshResizing/Source/MeshResizingCore/MeshResizingCore.Build.cs b/Engine/Plugins/Experimental/MeshResizing/Source/MeshResizingCore/MeshResizingCore.Build.cs
--- a/Engine/Plugins/Experimental/MeshResizing/Source/MeshResizingCore/MeshResizingCore.Build.cs
+++ b/Engine/Plugins/Experimental/MeshResizing/Source/MeshResizingCore/MeshResizingCore.Build.cs
@@ -15,6 +15,7 @@ namespace UnrealBuildTool.Rules
 					"CoreUObject",
 					"GeometryCore",
 					"HairStrandsCore",
+					"MeshDescription"
 				}
 			);
 		}
diff --git a/Engine/Plugins/Experimental/MeshResizing/Source/MeshResizingCore/Private/MeshResizing/RBFInterpolation.cpp b/Engine/Plugins/Experimental/MeshResizing/Source/MeshResizingCore/Private/MeshResizing/RBFInterpolation.cpp
--- a/Engine/Plugins/Experimental/MeshResizing/Source/MeshResizingCore/Private/MeshResizing/RBFInterpolation.cpp
+++ b/Engine/Plugins/Experimental/MeshResizing/Source/MeshResizingCore/Private/MeshResizing/RBFInterpolation.cpp
@@ -15,7 +15,12 @@ namespace UE::MeshResizing
 	{
 		static void UpdateInterpolationData(const GroomBinding_RBFWeighting::FPointsSampler& PointsSampler, const GroomBinding_RBFWeighting::FWeightsBuilder& InterpolationWeights, FMeshResizingRBFInterpolationData& OutData)
 		{
-			OutData.SampleIndices = PointsSampler.SampleIndices;
+			OutData.SampleIndices.SetNumUninitialized(PointsSampler.SampleIndices.Num());
+			for (int32 Index = 0; Index < PointsSampler.SampleIndices.Num(); ++Index)
+			{
+				OutData.SampleIndices[Index] = (int32)PointsSampler.SampleIndices[Index];
+			}
+
 			OutData.SampleRestPositions = PointsSampler.SamplePositions;
 			OutData.InterpolationWeights = InterpolationWeights.InverseEntries;
 		}
@@ -154,11 +159,42 @@ namespace UE::MeshResizing
 			Positions.Emplace(FVector3f(BaseMesh.GetVertexRef(VertexIndex)));
 		}
 
+		GenerateWeights(TConstArrayView<FVector3f>(Positions), NumInterpolationPoints, OutData);
+	}
+
+	void FRBFInterpolation::GenerateWeights(const FMeshDescription& BaseMesh, int32 NumInterpolationPoints, FMeshResizingRBFInterpolationData& OutData)
+	{
+		// Groom binding sampler expects points to be a contiguous array of FVector3f. Just copy that way for now.
+		TArray<FVector3f> Positions;
+		TArray<int32> CompactToVertID;
+		Positions.Reserve(BaseMesh.Vertices().Num());
+		CompactToVertID.Reserve(BaseMesh.Vertices().Num());
+		bool bCompactMatchesVertID = true;
+		for (const int32 VertexIndex : BaseMesh.Vertices().GetElementIDs())
+		{
+			Positions.Emplace(BaseMesh.GetVertexPosition(VertexIndex));
+			const int32 CompactIndex = CompactToVertID.Emplace(VertexIndex);
+			bCompactMatchesVertID = bCompactMatchesVertID && (CompactIndex == VertexIndex);
+		}
+
+		GenerateWeights(TConstArrayView<FVector3f>(Positions), NumInterpolationPoints, OutData);
+		if (!bCompactMatchesVertID)
+		{
+			// Update SampleIndices to be the true VertexIDs, not the CompactIDs
+			for (int32& SampleIndex : OutData.SampleIndices)
+			{
+				SampleIndex = CompactToVertID[SampleIndex];
+			}
+		}
+	}
+
+	void FRBFInterpolation::GenerateWeights(const TConstArrayView<FVector3f>& SourcePositions, int32 NumInterpolationPoints, FMeshResizingRBFInterpolationData& OutData)
+	{
 		// Consider all points to be valid for now
 		TArray<bool> ValidPoints;
-		ValidPoints.Init(true, BaseMesh.VertexCount());
+		ValidPoints.Init(true, SourcePositions.Num());
 
-		GroomBinding_RBFWeighting::FPointsSampler PointsSampler(ValidPoints, Positions.GetData(), NumInterpolationPoints);
+		GroomBinding_RBFWeighting::FPointsSampler PointsSampler(ValidPoints, SourcePositions.GetData(), NumInterpolationPoints);
 		const uint32 SampleCount = PointsSampler.SamplePositions.Num();
 		GroomBinding_RBFWeighting::FWeightsBuilder InterpolationWeights(SampleCount, SampleCount, PointsSampler.SamplePositions.GetData(), PointsSampler.SamplePositions.GetData());
 		Private::UpdateInterpolationData(PointsSampler, InterpolationWeights, OutData);
@@ -170,4 +206,25 @@ namespace UE::MeshResizing
 		Private::UpdateMeshSamples(TargetPositions, InterpolationData, SampleDeformations);
 		Private::ApplyRBFDeformation(InterpolationData, SampleDeformations, bInterpolateNormals, DeformingMesh);
 	}
+
+	void FRBFInterpolation::DeformPoints(const FMeshDescription& TargetMesh, const FMeshResizingRBFInterpolationData& InterpolationData, bool bInterpolateNormals, UE::Geometry::FDynamicMesh3& DeformingMesh)
+	{
+		TArray<FVector3f> TargetPositions;
+		TargetPositions.Reserve(InterpolationData.SampleIndices.Num());
+		for(int32 Index = 0; Index < InterpolationData.SampleIndices.Num(); ++Index)
+		{
+			const int32 SampleIndex = InterpolationData.SampleIndices[Index];
+			if (ensure(TargetMesh.IsVertexValid(SampleIndex)))
+			{
+				TargetPositions.Emplace(TargetMesh.GetVertexPosition(SampleIndex));
+			}
+			else
+			{
+				// Just use Source positions
+				TargetPositions.Emplace(InterpolationData.SampleRestPositions[Index]);
+			}
+		}
+
+		DeformPoints(TargetPositions, InterpolationData, bInterpolateNormals, DeformingMesh);
+	}
 }
\ No newline at end of file
diff --git a/Engine/Plugins/Experimental/MeshResizing/Source/MeshResizingCore/Public/MeshResizing/RBFInterpolation.h b/Engine/Plugins/Experimental/MeshResizing/Source/MeshResizingCore/Public/MeshResizing/RBFInterpolation.h
--- a/Engine/Plugins/Experimental/MeshResizing/Source/MeshResizingCore/Public/MeshResizing/RBFInterpolation.h
+++ b/Engine/Plugins/Experimental/MeshResizing/Source/MeshResizingCore/Public/MeshResizing/RBFInterpolation.h
@@ -9,6 +9,7 @@ namespace UE::Geometry
 {
 	class FDynamicMesh3;
 }
+struct FMeshDescription;
 
 USTRUCT()
 struct MESHRESIZINGCORE_API FMeshResizingRBFInterpolationData
@@ -16,7 +17,7 @@ struct MESHRESIZINGCORE_API FMeshResizingRBFInterpolationData
 	GENERATED_BODY()
 
 	UPROPERTY()
-	TArray<uint32> SampleIndices;
+	TArray<int32> SampleIndices;
 	UPROPERTY()
 	TArray<FVector3f> SampleRestPositions;
 	UPROPERTY()
@@ -27,9 +28,11 @@ namespace UE::MeshResizing
 {
 	struct MESHRESIZINGCORE_API FRBFInterpolation
 	{
-
+		static void GenerateWeights(const TConstArrayView<FVector3f>& SourcePositions, int32 NumInterpolationPoints, FMeshResizingRBFInterpolationData& OutData);
 		static void GenerateWeights(const UE::Geometry::FDynamicMesh3& BaseMesh, int32 NumInterpolationPoints, FMeshResizingRBFInterpolationData& OutData);
+		static void GenerateWeights(const FMeshDescription& BaseMesh, int32 NumInterpolationPoints, FMeshResizingRBFInterpolationData& OutData);
 		static void DeformPoints(const TArray<FVector3f>& TargetPositions, const FMeshResizingRBFInterpolationData& InterpolationData, bool bInterpolateNormals,  UE::Geometry::FDynamicMesh3& DeformingMesh);
+		static void DeformPoints(const FMeshDescription& TargetMesh, const FMeshResizingRBFInterpolationData& InterpolationData, bool bInterpolateNormals, UE::Geometry::FDynamicMesh3& DeformingMesh);
 	};
 }
 
diff --git a/Engine/Plugins/Experimental/MeshResizing/Source/MeshResizingNodes/MeshResizingDataflowNodes.Build.cs b/Engine/Plugins/Experimental/MeshResizing/Source/MeshResizingNodes/MeshResizingDataflowNodes.Build.cs
--- a/Engine/Plugins/Experimental/MeshResizing/Source/MeshResizingNodes/MeshResizingDataflowNodes.Build.cs
+++ b/Engine/Plugins/Experimental/MeshResizing/Source/MeshResizingNodes/MeshResizingDataflowNodes.Build.cs
@@ -27,6 +27,7 @@ public class MeshResizingDataflowNodes : ModuleRules
 				"DataflowEngine",
 				"DataflowNodes",
 				"GeometryFramework",
+				"MeshConversion",
 				"MeshConversionEngineTypes",
 				"MeshResizingCore",
 				"DynamicMesh",
diff --git a/Engine/Plugins/Experimental/MeshResizing/Source/MeshResizingNodes/Private/MeshResizing/BaseBodyDataflowNodes.cpp b/Engine/Plugins/Experimental/MeshResizing/Source/MeshResizingNodes/Private/MeshResizing/BaseBodyDataflowNodes.cpp
--- a/Engine/Plugins/Experimental/MeshResizing/Source/MeshResizingNodes/Private/MeshResizing/BaseBodyDataflowNodes.cpp
+++ b/Engine/Plugins/Experimental/MeshResizing/Source/MeshResizingNodes/Private/MeshResizing/BaseBodyDataflowNodes.cpp
@@ -7,6 +7,7 @@
 #include "Engine/SkeletalMesh.h"
 #include "Rendering/SkeletalMeshRenderData.h"
 #include "Rendering/SkeletalMeshModel.h"
+#include "MeshDescriptionToDynamicMesh.h"
 #include "SkeletalMeshLODModelToDynamicMesh.h"
 #include "SkeletalMeshLODRenderDataToDynamicMesh.h"
 #include "UDynamicMesh.h"
@@ -32,9 +33,9 @@ void FSkeletalMeshToMeshDataflowNode::Evaluate(UE::Dataflow::FContext& Context,
 	{
 		if (const USkeletalMesh* const InSkeletalMesh = GetValue(Context, &SkeletalMesh))
 		{
+#if WITH_EDITORONLY_DATA
 			if (bRecordImportedVertices)
 			{
-#if WITH_EDITORONLY_DATA
 				if (const FSkeletalMeshModel* const MeshModel = InSkeletalMesh->GetImportedModel())
 				{
 					if (MeshModel->LODModels.IsValidIndex(LODLevel))
@@ -74,11 +75,25 @@ void FSkeletalMeshToMeshDataflowNode::Evaluate(UE::Dataflow::FContext& Context,
 						return;
 					}
 				}
-#else
-				// Add error message about how this doesn't work? TBH not sure what to do here because if you started with the imported vertices and now they're gone, everything down stream is going to be broken
-#endif
+			}
+			else if (bUseMeshDescription)
+			{
+				if (const FMeshDescription* const MeshDescription = InSkeletalMesh->GetMeshDescription(LODLevel))
+				{
+					TObjectPtr<UDynamicMesh> NewMesh = NewObject<UDynamicMesh>();
+					NewMesh->Reset();
+
+					UE::Geometry::FDynamicMesh3& DynMesh = NewMesh->GetMeshRef();
+					FMeshDescriptionToDynamicMesh Converter;
+					Converter.bVIDsFromNonManifoldMeshDescriptionAttr = true;
+					Converter.Convert(MeshDescription, DynMesh, true);
+
+					SetValue(Context, NewMesh, &Mesh);
+					return;
+				}
 			}
 			else
+#endif
 			{
 				if (const FSkeletalMeshRenderData* const RenderData = InSkeletalMesh->GetResourceForRendering())
 				{
diff --git a/Engine/Plugins/Experimental/MeshResizing/Source/MeshResizingNodes/Private/MeshResizing/BaseBodyDataflowNodes.h b/Engine/Plugins/Experimental/MeshResizing/Source/MeshResizingNodes/Private/MeshResizing/BaseBodyDataflowNodes.h
--- a/Engine/Plugins/Experimental/MeshResizing/Source/MeshResizingNodes/Private/MeshResizing/BaseBodyDataflowNodes.h
+++ b/Engine/Plugins/Experimental/MeshResizing/Source/MeshResizingNodes/Private/MeshResizing/BaseBodyDataflowNodes.h
@@ -54,9 +54,15 @@ private:
 	UPROPERTY(EditAnywhere, Category = "SkeletalMesh", meta = (DisplayName = "LOD Level"));
 	int32 LODLevel = 0;
 
+#if WITH_EDITORONLY_DATA
 	/** Record ImportedVertices (if available) as NonManifold mapping data. This requires Editor-Only data.*/
 	UPROPERTY(EditAnywhere, Category = "SkeletalMesh")
 	bool bRecordImportedVertices = true;
+
+	/** Generate from mesh description. Requires Editor-Only data.*/
+	UPROPERTY(EditAnywhere, Category = "SkeletalMesh")
+	bool bUseMeshDescription = false;
+#endif
 };
 
 
diff --git a/Engine/Plugins/Experimental/MeshResizing/Source/MeshResizingNodes/Private/MeshResizing/RBFInterpolationNodes.cpp b/Engine/Plugins/Experimental/MeshResizing/Source/MeshResizingNodes/Private/MeshResizing/RBFInterpolationNodes.cpp
--- a/Engine/Plugins/Experimental/MeshResizing/Source/MeshResizingNodes/Private/MeshResizing/RBFInterpolationNodes.cpp
+++ b/Engine/Plugins/Experimental/MeshResizing/Source/MeshResizingNodes/Private/MeshResizing/RBFInterpolationNodes.cpp
@@ -11,6 +11,9 @@
 #if WITH_EDITOR
 #include "Dataflow/DataflowRenderingViewMode.h"
 #endif
+#if WITH_EDITORONLY_DATA
+#include "Engine/SkeletalMesh.h"
+#endif
 
 #include UE_INLINE_GENERATED_CPP_BY_NAME(RBFInterpolationNodes)
 
@@ -19,8 +22,9 @@ using UE::Geometry::FDynamicMesh3;
 FGenerateRBFResizingWeightsNode::FGenerateRBFResizingWeightsNode(const UE::Dataflow::FNodeParameters& InParam, FGuid InGuid) :
 	FDataflowNode(InParam, InGuid)
 {
-	RegisterInputConnection(&MeshToResize);
 	RegisterInputConnection(&SourceMesh);
+	RegisterInputConnection(&NumInterpolationPoints);
+	RegisterInputConnection(&MeshToResize);
 
 	RegisterOutputConnection(&InterpolationData);
 }
@@ -31,14 +35,11 @@ void FGenerateRBFResizingWeightsNode::Evaluate(UE::Dataflow::FContext& Context,
 	if (Out->IsA<FMeshResizingRBFInterpolationData>(&InterpolationData))
 	{
 		FMeshResizingRBFInterpolationData Result;
-		if (TObjectPtr<UDataflowMesh> InMeshToResize = GetValue(Context, &MeshToResize))
+		if (TObjectPtr<UDataflowMesh> InSourceMesh = GetValue(Context, &SourceMesh))
 		{
-			if (TObjectPtr<UDataflowMesh> InSourceMesh = GetValue(Context, &SourceMesh))
+			if (InSourceMesh->GetDynamicMesh())
 			{
-				if (InMeshToResize->GetDynamicMesh() && InSourceMesh->GetDynamicMesh())
-				{
-					UE::MeshResizing::FRBFInterpolation::GenerateWeights(InSourceMesh->GetDynamicMeshRef(), NumInterpolationPoints, Result);
-				}
+				UE::MeshResizing::FRBFInterpolation::GenerateWeights(InSourceMesh->GetDynamicMeshRef(), GetValue(Context, &NumInterpolationPoints), Result);
 			}
 		}
 
@@ -76,6 +77,10 @@ FApplyRBFResizingNode::FApplyRBFResizingNode(const UE::Dataflow::FNodeParameters
 	FDataflowNode(InParam, InGuid)
 {
 	RegisterInputConnection(&MeshToResize);
+#if WITH_EDITORONLY_DATA
+	RegisterInputConnection(&TargetSkeletalMesh);
+	RegisterInputConnection(&TargetSkeletalMeshLODIndex);
+#endif
 	RegisterInputConnection(&TargetMesh);
 	RegisterInputConnection(&InterpolationData);
 
@@ -85,33 +90,62 @@ FApplyRBFResizingNode::FApplyRBFResizingNode(const UE::Dataflow::FNodeParameters
 
 void FApplyRBFResizingNode::Evaluate(UE::Dataflow::FContext& Context, const FDataflowOutput* Out) const
 {
-	if (TObjectPtr<UDataflowMesh> InMeshToResize = GetValue(Context, &MeshToResize))
+	if (TObjectPtr<const UDataflowMesh> InMeshToResize = GetValue(Context, &MeshToResize))
 	{
-		if (TObjectPtr<UDataflowMesh> InTargetMesh = GetValue(Context, &TargetMesh))
+#if WITH_EDITORONLY_DATA
+		if (bUseSkeletalMeshTarget)
 		{
-			const FMeshResizingRBFInterpolationData& InInterpolationData = GetValue(Context, &InterpolationData);
-
-			if (InMeshToResize->GetDynamicMesh() && InTargetMesh->GetDynamicMesh() && InInterpolationData.SampleIndices.Num())
+			if (TObjectPtr<const USkeletalMesh> InTargetSKM = GetValue(Context, &TargetSkeletalMesh))
 			{
-				TObjectPtr<UDataflowMesh> OutResizedMesh = NewObject<UDataflowMesh>();
-				FDynamicMesh3 ResizeMesh;
-				ResizeMesh.Copy(InMeshToResize->GetDynamicMeshRef());
-
-				const FDynamicMesh3& TargetMeshRef = InTargetMesh->GetDynamicMeshRef();
+				const int32 InLODIndex = GetValue(Context, &TargetSkeletalMeshLODIndex);
+				if (const FMeshDescription* const MeshDescription = InTargetSKM->GetMeshDescription(InLODIndex))
+				{
+					const FMeshResizingRBFInterpolationData& InInterpolationData = GetValue(Context, &InterpolationData);
+					if (InMeshToResize->GetDynamicMesh() && InInterpolationData.SampleIndices.Num())
+					{
+						TObjectPtr<UDataflowMesh> OutResizedMesh = NewObject<UDataflowMesh>();
+						FDynamicMesh3 ResizeMesh;
+						ResizeMesh.Copy(InMeshToResize->GetDynamicMeshRef());
+
+						// Do the interpolation 
+						UE::MeshResizing::FRBFInterpolation::DeformPoints(*MeshDescription, InInterpolationData, bInterpolateNormals, ResizeMesh);
+						OutResizedMesh->SetDynamicMesh(MoveTemp(ResizeMesh));
+						OutResizedMesh->SetMaterials(InMeshToResize->GetMaterials());
+						SetValue<TObjectPtr<const UDataflowMesh>>(Context, OutResizedMesh, &ResizedMesh);
+						return;
+					}
+				}
+			}
+		}
+		else
+#endif
+		{
+			if (TObjectPtr<const UDataflowMesh> InTargetMesh = GetValue(Context, &TargetMesh))
+			{
+				const FMeshResizingRBFInterpolationData& InInterpolationData = GetValue(Context, &InterpolationData);
 
-				// Get target mesh samples
-				TArray<FVector3f> TargetSamplePoints;
-				TargetSamplePoints.Reserve(InInterpolationData.SampleIndices.Num());
-				for (const uint32 SampleIndex : InInterpolationData.SampleIndices)
+				if (InMeshToResize->GetDynamicMesh() && InTargetMesh->GetDynamicMesh() && InInterpolationData.SampleIndices.Num())
 				{
-					TargetSamplePoints.Emplace(FVector3f(TargetMeshRef.GetVertex(SampleIndex)));
+					TObjectPtr<UDataflowMesh> OutResizedMesh = NewObject<UDataflowMesh>();
+					FDynamicMesh3 ResizeMesh;
+					ResizeMesh.Copy(InMeshToResize->GetDynamicMeshRef());
+
+					const FDynamicMesh3& TargetMeshRef = InTargetMesh->GetDynamicMeshRef();
+
+					// Get target mesh samples
+					TArray<FVector3f> TargetSamplePoints;
+					TargetSamplePoints.Reserve(InInterpolationData.SampleIndices.Num());
+					for (const uint32 SampleIndex : InInterpolationData.SampleIndices)
+					{
+						TargetSamplePoints.Emplace(FVector3f(TargetMeshRef.GetVertex(SampleIndex)));
+					}
+					// Do the interpolation
+					UE::MeshResizing::FRBFInterpolation::DeformPoints(TargetSamplePoints, InInterpolationData, bInterpolateNormals, ResizeMesh);
+					OutResizedMesh->SetDynamicMesh(MoveTemp(ResizeMesh));
+					OutResizedMesh->SetMaterials(InMeshToResize->GetMaterials());
+					SetValue<TObjectPtr<const UDataflowMesh>>(Context, OutResizedMesh, &ResizedMesh);
+					return;
 				}
-				// Do the interpolation
-				UE::MeshResizing::FRBFInterpolation::DeformPoints(TargetSamplePoints, InInterpolationData, bInterpolateNormals, ResizeMesh);
-				OutResizedMesh->SetDynamicMesh(MoveTemp(ResizeMesh));
-				OutResizedMesh->SetMaterials(InMeshToResize->GetMaterials());
-				SetValue(Context, OutResizedMesh, &ResizedMesh);
-				return;
 			}
 		}
 	}
diff --git a/Engine/Plugins/Experimental/MeshResizing/Source/MeshResizingNodes/Private/MeshResizing/RBFInterpolationNodes.h b/Engine/Plugins/Experimental/MeshResizing/Source/MeshResizingNodes/Private/MeshResizing/RBFInterpolationNodes.h
--- a/Engine/Plugins/Experimental/MeshResizing/Source/MeshResizingNodes/Private/MeshResizing/RBFInterpolationNodes.h
+++ b/Engine/Plugins/Experimental/MeshResizing/Source/MeshResizingNodes/Private/MeshResizing/RBFInterpolationNodes.h
@@ -10,6 +10,7 @@
 
 class UDataflowMesh;
 class UMaterialInterface;
+class USkeletalMesh;
 
 /**
  * Sample points and generate RBF Interpolation data for a given Source mesh.
@@ -39,8 +40,8 @@ private:
 	FMeshResizingRBFInterpolationData InterpolationData;
 
 	/** The number of interpolation points to be sampled. */
-	UPROPERTY(EditAnywhere, Category = "RBF Resizing Weights", meta = (Min = 1))
-	int32 NumInterpolationPoints = 1000;
+	UPROPERTY(EditAnywhere, Category = "RBF Resizing Weights", meta = (Min = 1, DataflowInput))
+	int32 NumInterpolationPoints = 1500;
 
 	//~ Begin FDataflowNode interface
 	virtual void Evaluate(UE::Dataflow::FContext& Context, const FDataflowOutput* Out) const override; 
@@ -72,11 +73,28 @@ private:
 
 	/** The mesh being resized */
 	UPROPERTY(meta = (DataflowInput))
-	TObjectPtr<UDataflowMesh> MeshToResize;
+	TObjectPtr<const UDataflowMesh> MeshToResize;
+
+#if WITH_EDITORONLY_DATA
+	// Skeletal mesh target requires access to the MeshDescription, which is Editor-only
+	/** Use a skeletal mesh for the target mesh (instead of a dynamic mesh)*/
+	UPROPERTY(EditAnywhere, Category = "RBF Resize")
+	bool bUseSkeletalMeshTarget = false;
 
 	/** The target mesh that corresponds with the SourceMesh used to generate the InterpolationData. Must have matching vertices with SourceMesh */
-	UPROPERTY(meta = (DataflowInput))
-	TObjectPtr<UDataflowMesh> TargetMesh;
+	UPROPERTY(meta = (DataflowInput, EditCondition = "UseSkeletalMeshTarget"))
+	TObjectPtr<const USkeletalMesh> TargetSkeletalMesh;
+
+	UPROPERTY(EditAnywhere, Category = "RBF Resize", meta = (DataflowInput, EditCondition = "UseSkeletalMeshTarget", ClampMin = 0))
+	int32 TargetSkeletalMeshLODIndex = 0;
+
+#else
+	static constexpr bool bUseSkeletalMeshTarget = false;
+#endif
+
+	/** The target mesh that corresponds with the SourceMesh used to generate the InterpolationData. Must have matching vertices with SourceMesh */
+	UPROPERTY(meta = (DataflowInput, EditCondition = "!bUseSkeletalMeshTarget"))
+	TObjectPtr<const UDataflowMesh> TargetMesh;
 
 	/** The pre-calculated base RBF interpolation data.*/
 	UPROPERTY(meta = (DataflowInput))
@@ -84,7 +102,7 @@ private:
 
 	/** The resulting resized mesh */
 	UPROPERTY(meta = (DataflowOutput, DataflowPassthrough = "MeshToResize"))
-	TObjectPtr<UDataflowMesh> ResizedMesh;
+	TObjectPtr<const UDataflowMesh> ResizedMesh;
 
 	/** Whether or not to interpolate the normals as well as the positions*/
 	UPROPERTY(EditAnywhere, Category = "RBF Resize")
