Chaos Outfit Asset - Added SkeletalMesh export.
Note the original commit was backed out at CL 41823077 due to a Developer folder dependency issue that has been addressed here.

#rnx

diff --git a/Engine/Plugins/ChaosClothAsset/Source/ChaosClothAssetEngine/Private/ChaosClothAsset/ClothAssetBase.cpp b/Engine/Plugins/ChaosClothAsset/Source/ChaosClothAssetEngine/Private/ChaosClothAsset/ClothAssetBase.cpp
--- a/Engine/Plugins/ChaosClothAsset/Source/ChaosClothAssetEngine/Private/ChaosClothAsset/ClothAssetBase.cpp
+++ b/Engine/Plugins/ChaosClothAsset/Source/ChaosClothAssetEngine/Private/ChaosClothAsset/ClothAssetBase.cpp
@@ -6,13 +6,15 @@
 #endif
 #include "ChaosClothAsset/ClothAssetPrivate.h"
 #include "ChaosClothAsset/ClothComponent.h"
+#include "ChaosClothAsset/SkeletalMeshConverterClassProvider.h"
 #include "Engine/Engine.h"
 #include "Engine/SkeletalMesh.h"
 #include "Engine/SkinnedAssetAsyncCompileUtils.h"
-#include "Rendering/SkeletalMeshRenderData.h"
 #if WITH_EDITOR
+#include "Features/IModularFeatures.h"
 #include "Interfaces/ITargetPlatform.h"
 #endif
+#include "Rendering/SkeletalMeshRenderData.h"
 #include "UObject/FortniteMainBranchObjectVersion.h"
 
 
@@ -484,6 +486,32 @@ PRAGMA_DISABLE_DEPRECATION_WARNINGS
 PRAGMA_ENABLE_DEPRECATION_WARNINGS
 }
 
+#if WITH_EDITOR
+bool UChaosClothAssetBase::ExportToSkeletalMesh(USkeletalMesh& SkeletalMesh) const
+{
+	using namespace UE::Chaos::ClothAsset;
+	using namespace UE::Geometry;
+
+	const TArray<IClothAssetSkeletalMeshConverterClassProvider*> ClassProviders =
+		IModularFeatures::Get().GetModularFeatureImplementations<IClothAssetSkeletalMeshConverterClassProvider>(IClothAssetSkeletalMeshConverterClassProvider::FeatureName);
+	if (const IClothAssetSkeletalMeshConverterClassProvider* const ClassProvider = ClassProviders.Num() ? ClassProviders[0] : nullptr)
+	{
+		if (const TSubclassOf<UClothAssetSkeletalMeshConverter> ClothAssetSkeletalMeshConverterClass = ClassProvider->GetClothAssetSkeletalMeshConverter())
+		{
+			if (UClothAssetSkeletalMeshConverter* const ClothAssetSkeletalMeshConverter = ClothAssetSkeletalMeshConverterClass->GetDefaultObject<UClothAssetSkeletalMeshConverter>())
+			{
+				return ClothAssetSkeletalMeshConverter->ExportToSkeletalMesh(*this, SkeletalMesh);
+			}
+		}
+	}
+	else
+	{
+		UE_LOG(LogChaosClothAsset, Error, TEXT("The export to SkeletalMesh has failed: Cannot find a SkeletalMesh converter. Make sure to enable the ChaosClothAssetEditor plugin."));
+	}
+	return false;
+}
+#endif  // #if WITH_EDITOR
+
 void UChaosClothAssetBase::ReregisterComponents()
 {
 	// Recreate the simulation proxies with the updated physics asset
diff --git a/Engine/Plugins/ChaosClothAsset/Source/ChaosClothAssetEngine/Public/ChaosClothAsset/ClothAssetBase.h b/Engine/Plugins/ChaosClothAsset/Source/ChaosClothAssetEngine/Public/ChaosClothAsset/ClothAssetBase.h
--- a/Engine/Plugins/ChaosClothAsset/Source/ChaosClothAssetEngine/Public/ChaosClothAsset/ClothAssetBase.h
+++ b/Engine/Plugins/ChaosClothAsset/Source/ChaosClothAssetEngine/Public/ChaosClothAsset/ClothAssetBase.h
@@ -260,6 +260,16 @@ public:
 
 	void SetHasVertexColors(bool InbHasVertexColors);
 
+#if WITH_EDITOR
+	/**
+	 * Export the graphical representation of this asset to a SkeletalMesh asset.
+	 * Includes skinning, but excludes all clothing simulation data since it isn't compatible with the Cloth Asset.
+	 * @param SkeletalMesh The empty skeletal mesh asset to be written to.
+	 * @return Whether the initialization of the skeletal mesh has succeeded or not.
+	 */
+	bool ExportToSkeletalMesh(USkeletalMesh& SkeletalMesh) const;
+#endif
+
 protected:
 	/** Reregister all components using this asset to reset the simulation in case anything has changed. */
 	void ReregisterComponents();
diff --git a/Engine/Plugins/ChaosClothAsset/Source/ChaosClothAssetEngine/Public/ChaosClothAsset/SkeletalMeshConverterClassProvider.h b/Engine/Plugins/ChaosClothAsset/Source/ChaosClothAssetEngine/Public/ChaosClothAsset/SkeletalMeshConverterClassProvider.h
new file mode 100644
--- /dev/null
+++ b/Engine/Plugins/ChaosClothAsset/Source/ChaosClothAssetEngine/Public/ChaosClothAsset/SkeletalMeshConverterClassProvider.h
@@ -0,0 +1,31 @@
+// Copyright Epic Games, Inc. All Rights Reserved.
+
+#pragma once
+
+#include "Features/IModularFeature.h"
+#include "SkeletalMeshConverterClassProvider.generated.h"
+
+/** Modular base class for skeletal mesh conversion. */
+UCLASS(Abstract)
+class CHAOSCLOTHASSETENGINE_API UClothAssetSkeletalMeshConverter : public UObject
+{
+	GENERATED_BODY()
+
+public:
+	/** Build a skeletal mesh from the specified Cloth Asset. */
+	virtual bool ExportToSkeletalMesh(const class UChaosClothAssetBase& ClothAssetBase, class USkeletalMesh& SkeletalMesh) const
+		PURE_VIRTUAL(UClothAssetSkeletalMeshConverter::ExportToSkeletalMesh, return false;);
+};
+
+/**
+ * A modular interface to provide ways to build a SkeletalMesh from a DynamicMesh, which can be used to build a SkeletalMesh from a ClothCollection.
+ * This cannot be done in a non Editor module due to the required dependency being in a Developper folder,
+ * and instead is exposed as a modular feature in order to be called from the UChaosClothAssetBase engine class.
+ */
+class IClothAssetSkeletalMeshConverterClassProvider : public IModularFeature
+{
+public:
+	inline static const FName FeatureName = TEXT("IClothAssetSkeletalMeshConverterClassProvider");
+
+	virtual TSubclassOf<UClothAssetSkeletalMeshConverter> GetClothAssetSkeletalMeshConverter() const = 0;
+};
diff --git a/Engine/Plugins/ChaosClothAssetEditor/Source/ChaosClothAssetTools/ChaosClothAssetTools.Build.cs b/Engine/Plugins/ChaosClothAssetEditor/Source/ChaosClothAssetTools/ChaosClothAssetTools.Build.cs
--- a/Engine/Plugins/ChaosClothAssetEditor/Source/ChaosClothAssetTools/ChaosClothAssetTools.Build.cs
+++ b/Engine/Plugins/ChaosClothAssetEditor/Source/ChaosClothAssetTools/ChaosClothAssetTools.Build.cs
@@ -31,7 +31,10 @@ public class ChaosClothAssetTools : ModuleRules
 				"ClothPainter",                 // For clothing asset exporter
 				"ClothingSystemRuntimeCommon",
 				"GeometryCore",					// For DynamicMesh conversion
+				"MeshConversion",
+				"MeshDescription",
 				"SkeletalMeshDescription",		// For FSkeletalMeshAttributes::DefaultSkinWeightProfileName
+				"SkeletalMeshUtilitiesCommon"           
 			}
 		);
 	}
diff --git a/Engine/Plugins/ChaosClothAssetEditor/Source/ChaosClothAssetTools/Private/ChaosClothAsset/ChaosClothAssetToolsModule.cpp b/Engine/Plugins/ChaosClothAssetEditor/Source/ChaosClothAssetTools/Private/ChaosClothAsset/ChaosClothAssetToolsModule.cpp
--- a/Engine/Plugins/ChaosClothAssetEditor/Source/ChaosClothAssetTools/Private/ChaosClothAsset/ChaosClothAssetToolsModule.cpp
+++ b/Engine/Plugins/ChaosClothAssetEditor/Source/ChaosClothAssetTools/Private/ChaosClothAsset/ChaosClothAssetToolsModule.cpp
@@ -3,12 +3,13 @@
 #include "Features/IModularFeatures.h"
 #include "Modules/ModuleManager.h"
 #include "ChaosClothAsset/ClothAssetEditorStyle.h"
-#include "ClothComponentEditorStyle.h"
-#include "ClothingAssetToClothAssetExporter.h"
+#include "ChaosClothAsset/ClothComponentEditorStyle.h"
+#include "ChaosClothAsset/ClothingAssetToClothAssetExporter.h"
+#include "ChaosClothAsset/SkeletalMeshConverter.h"
 
 namespace UE::Chaos::ClothAsset
 {
-	class FChaosClothAssetToolsModule : public IModuleInterface, public IClothingAssetExporterClassProvider
+	class FChaosClothAssetToolsModule : public IModuleInterface, public IClothingAssetExporterClassProvider, public IClothAssetSkeletalMeshConverterClassProvider
 	{
 	public:
 		// IModuleInterface implementation
@@ -20,6 +21,7 @@ namespace UE::Chaos::ClothAsset
 
 			// Register modular features
 			IModularFeatures::Get().RegisterModularFeature(IClothingAssetExporterClassProvider::FeatureName, static_cast<IClothingAssetExporterClassProvider*>(this));
+			IModularFeatures::Get().RegisterModularFeature(IClothAssetSkeletalMeshConverterClassProvider::FeatureName, static_cast<IClothAssetSkeletalMeshConverterClassProvider*>(this));
 		}
 
 		virtual void ShutdownModule() override
@@ -28,14 +30,23 @@ namespace UE::Chaos::ClothAsset
 			{
 				// Unregister modular features
 				IModularFeatures::Get().UnregisterModularFeature(IClothingAssetExporterClassProvider::FeatureName, static_cast<IClothingAssetExporterClassProvider*>(this));
+				IModularFeatures::Get().UnregisterModularFeature(IClothAssetSkeletalMeshConverterClassProvider::FeatureName, static_cast<IClothAssetSkeletalMeshConverterClassProvider*>(this));
 			}
 		}
 
-		// IClothingSimulationFactoryClassProvider implementation
+		//~ Begin IClothingSimulationFactoryClassProvider implementation
 		virtual TSubclassOf<UClothingAssetExporter> GetClothingAssetExporterClass() const override
 		{
 			return UClothingAssetToChaosClothAssetExporter::StaticClass();
 		}
+		//~ End IClothingSimulationFactoryClassProvider implementation
+
+		//~ Begin IClothAssetSkeletalMeshConverterClassProvider implementation
+		virtual TSubclassOf<UClothAssetSkeletalMeshConverter> GetClothAssetSkeletalMeshConverter() const override
+		{
+			return UClothAssetEditorSkeletalMeshConverter::StaticClass();
+		}
+		//~ End IClothAssetSkeletalMeshConverterClassProvider implementation
 	};
 }
 
diff --git a/Engine/Plugins/ChaosClothAssetEditor/Source/ChaosClothAssetTools/Private/ChaosClothAsset/ClothPatternToDynamicMesh.cpp b/Engine/Plugins/ChaosClothAssetEditor/Source/ChaosClothAssetTools/Private/ChaosClothAsset/ClothPatternToDynamicMesh.cpp
--- a/Engine/Plugins/ChaosClothAssetEditor/Source/ChaosClothAssetTools/Private/ChaosClothAsset/ClothPatternToDynamicMesh.cpp
+++ b/Engine/Plugins/ChaosClothAssetEditor/Source/ChaosClothAssetTools/Private/ChaosClothAsset/ClothPatternToDynamicMesh.cpp
@@ -625,7 +625,13 @@ private:
 };
 
 
-void FClothPatternToDynamicMesh::Convert(const TSharedRef<const FManagedArrayCollection> ClothCollection, int32 PatternIndex, EClothPatternVertexType VertexDataType, UE::Geometry::FDynamicMesh3& MeshOut, bool bDisableAttributes)
+void FClothPatternToDynamicMesh::Convert(
+	const TSharedRef<const FManagedArrayCollection> ClothCollection,
+	int32 PatternIndex,
+	EClothPatternVertexType VertexDataType,
+	UE::Geometry::FDynamicMesh3& MeshOut,
+	bool bDisableAttributes,
+	int32 MaterialOffset)
 {
 	const FCollectionClothConstFacade ClothFacade(ClothCollection);
 
@@ -703,14 +709,42 @@ void FClothPatternToDynamicMesh::Convert(const TSharedRef<const FManagedArrayCol
 		MeshOut.EnableAttributes();
 		FClothPatternToDynamicMeshMappingSupport::AttachTriangleMappingData(PatternToDynamicMesh.ToSrcTriIDMap, MeshOut);
 	}
+
+	// Fix up any triangles without valid material IDs and apply material offset
+	if (VertexDataType == EClothPatternVertexType::Render && MaterialOffset != INDEX_NONE)
+	{
+		const TConstArrayView<FString> MaterialPaths = ClothFacade.GetRenderMaterialPathName();
+		int32 DefaultMaterialID = 0;
+		for (const int32 TriID : MeshOut.TriangleIndicesItr())
+		{
+			int32 MaterialID;
+			MeshOut.Attributes()->GetMaterialID()->GetValue(TriID, &MaterialID);
+			//const int32 MaterialID = MeshOut.Attributes()->GetMaterialID()->GetValue(TriID);
+			if (!MaterialPaths.IsValidIndex(MaterialID))
+			{
+				MeshOut.Attributes()->GetMaterialID()->SetValue(TriID, DefaultMaterialID);
+			}
+			else
+			{
+				MeshOut.Attributes()->GetMaterialID()->SetValue(TriID, MaterialID + MaterialOffset);
+			}
+		}
+	}
 }
 
-void FClothPatternToDynamicMesh::Convert(const UChaosClothAsset* ClothAssetMeshIn, int32 LODIndex, int32 PatternIndex, EClothPatternVertexType VertexDataType, UE::Geometry::FDynamicMesh3& MeshOut)
+void FClothPatternToDynamicMesh::Convert(
+	const UChaosClothAsset* ClothAssetMeshIn,
+	int32 LODIndex,
+	int32 PatternIndex,
+	EClothPatternVertexType VertexDataType,
+	UE::Geometry::FDynamicMesh3& MeshOut,
+	int32 MaterialOffset)
 {
 	const TArray<TSharedRef<const FManagedArrayCollection>>& ClothCollections = ClothAssetMeshIn->GetClothCollections();
 	check(ClothCollections.IsValidIndex(LODIndex));
 
-	Convert(ClothCollections[LODIndex], PatternIndex, VertexDataType, MeshOut);
+	constexpr bool bDisableAttributes = false;
+	Convert(ClothCollections[LODIndex], PatternIndex, VertexDataType, MeshOut, bDisableAttributes, MaterialOffset);
 }
 
 }	// namespace UE::Chaos::ClothAsset
@@ -719,8 +753,24 @@ void FClothPatternToDynamicMesh::Convert(const UChaosClothAsset* ClothAssetMeshI
 
 namespace UE::Chaos::ClothAsset
 {
+void FClothPatternToDynamicMesh::Convert(
+	const TSharedRef<const FManagedArrayCollection> ClothCollection,
+	int32 PatternIndex,
+	EClothPatternVertexType VertexDataType,
+	UE::Geometry::FDynamicMesh3& MeshOut,
+	bool bDisableAttributes,
+	int32 MaterialOffset)
+{
+	// Conversion only supported with editor.
+	check(0);
+}
 
-void FClothAssetToDynamicMesh::Convert(const UChaosClothAsset* ClothAssetMeshIn, FDynamicMesh3& MeshOut, int32 LODIndex, int32 PatternIndex)
+void FClothPatternToDynamicMesh::Convert(
+	const UChaosClothAsset* ClothAssetMeshIn,
+	FDynamicMesh3& MeshOut,
+	int32 LODIndex,
+	int32 PatternIndex,
+	int32 MaterialOffset)
 {
 	// Conversion only supported with editor.
 	check(0);
diff --git a/Engine/Plugins/ChaosClothAssetEditor/Source/ChaosClothAssetTools/Private/ChaosClothAsset/SkeletalMeshConverter.cpp b/Engine/Plugins/ChaosClothAssetEditor/Source/ChaosClothAssetTools/Private/ChaosClothAsset/SkeletalMeshConverter.cpp
new file mode 100644
--- /dev/null
+++ b/Engine/Plugins/ChaosClothAssetEditor/Source/ChaosClothAssetTools/Private/ChaosClothAsset/SkeletalMeshConverter.cpp
@@ -0,0 +1,124 @@
+// Copyright Epic Games, Inc. All Rights Reserved.
+
+#include "ChaosClothAsset/SkeletalMeshConverter.h"
+#include "ChaosClothAsset/ClothAssetBase.h"
+#include "ChaosClothAsset/ClothPatternToDynamicMesh.h"
+#include "ChaosClothAsset/CollectionClothFacade.h"
+#include "DynamicMesh/DynamicMesh3.h"
+#include "DynamicMesh/DynamicMeshAttributeSet.h"
+#include "Engine/SkeletalMesh.h"
+#include "Materials/Material.h"
+#include "Rendering/SkeletalMeshModel.h"
+#include "DynamicMeshToMeshDescription.h"
+#include "MaterialDomain.h"
+#include "SkeletalMeshAttributes.h"
+#include "StaticToSkeletalMeshConverter.h"
+
+bool UClothAssetEditorSkeletalMeshConverter::ExportToSkeletalMesh(const UChaosClothAssetBase& ClothAssetBase, USkeletalMesh& SkeletalMesh) const
+{
+	using namespace UE::Chaos::ClothAsset;
+
+	// Create one dynamic mesh from the multiple collections
+	TArray<UMaterialInterface*> RenderMaterials;
+	TArray<FDynamicMesh3> DynamicMeshes;
+	DynamicMeshes.AddDefaulted();  // TODO: Do the other LODs, one per Dynamic mesh
+
+	int32 MaterialOffset = 0;
+
+	const int32 NumModels = ClothAssetBase.GetNumClothSimulationModels();
+
+	for (int32 ModelIndex = 0; ModelIndex < NumModels; ++ModelIndex)
+	{
+		const TArray<TSharedRef<const FManagedArrayCollection>>& ClothCollections = ClothAssetBase.GetCollections(ModelIndex);
+
+		//for (const TSharedRef<const FManagedArrayCollection>& ClothCollection : ClothCollections)
+		if (ClothCollections.Num())
+		{
+			const TSharedRef<const FManagedArrayCollection>& ClothCollection = ClothCollections[0];  // TODO: Do other LODs
+			FCollectionClothConstFacade ClothFacade(ClothCollection);
+
+			if (ClothFacade.IsValid())
+			{
+				FDynamicMesh3 DynamicMeshPiece;
+
+				constexpr bool bDisableAttributes = false;
+				const int32 PatternIndex = INDEX_NONE;
+				DynamicMeshPiece.EnableAttributes();
+				DynamicMeshPiece.Attributes()->EnableMaterialID();
+				FClothPatternToDynamicMesh ClothPatternToDynamicMesh;
+				ClothPatternToDynamicMesh.Convert(
+					ClothCollection,
+					PatternIndex,
+					EClothPatternVertexType::Render,
+					DynamicMeshPiece,
+					bDisableAttributes,
+					MaterialOffset);
+
+				if (ModelIndex == 0)
+				{
+					DynamicMeshes[0] = MoveTemp(DynamicMeshPiece);
+				}
+				else
+				{
+					DynamicMeshes[0].AppendWithOffsets(DynamicMeshPiece);
+				}
+
+				const TArrayView<const FString> MaterialPaths = ClothFacade.GetRenderMaterialPathName();
+				RenderMaterials.Reserve(RenderMaterials.Num() + MaterialPaths.Num());
+				for (int32 MaterialIndex = 0; MaterialIndex < MaterialPaths.Num(); ++MaterialIndex)
+				{
+					RenderMaterials.Emplace(LoadObject<UMaterialInterface>(nullptr, *MaterialPaths[MaterialIndex]));
+				}
+				MaterialOffset += MaterialPaths.Num();
+			}
+		}
+	}
+
+	// Create the skeletal mesh from the dynamic mesh
+	TArray<const FMeshDescription*> MeshDescriptions;
+
+	SkeletalMesh.PreEditChange(nullptr);
+	SkeletalMesh.GetImportedModel()->LODModels.Empty();
+	SkeletalMesh.ResetLODInfo();
+
+	TArray<FMeshDescription> LocalDescriptions;
+	for (const FDynamicMesh3& DynamicMesh : DynamicMeshes)
+	{
+		// Create a mesh description
+		FMeshDescription& MeshDescription = LocalDescriptions.AddDefaulted_GetRef();
+
+		// Add skeletal mesh attributes to the mesh description
+		FSkeletalMeshAttributes Attributes(MeshDescription);
+		Attributes.Register();
+
+		// Convert dynamic mesh to the mesh description
+		FConversionToMeshDescriptionOptions ConverterOptions;
+		FDynamicMeshToMeshDescription Converter(ConverterOptions);
+		constexpr bool bCopyTangents = true;
+		Converter.Convert(&DynamicMesh, MeshDescription, bCopyTangents);
+
+		// Add the created description to the list
+		MeshDescriptions.Add(&MeshDescription);
+	}
+
+	TArray<FSkeletalMaterial> SkeletalMaterials;
+	SkeletalMaterials.Reserve(RenderMaterials.Num());
+	for (UMaterialInterface* RenderMaterial : RenderMaterials)
+	{
+		SkeletalMaterials.Emplace(RenderMaterial, RenderMaterial->GetFName());
+	}
+	if (!SkeletalMaterials.Num())
+	{
+		SkeletalMaterials.Emplace(UMaterial::GetDefaultMaterial(MD_Surface));
+	}
+
+	constexpr bool bRecomputeNormals = false;
+	constexpr bool bRecomputeTangents = false;
+	return FStaticToSkeletalMeshConverter::InitializeSkeletalMeshFromMeshDescriptions(
+		&SkeletalMesh,
+		MeshDescriptions,
+		SkeletalMaterials,
+		ClothAssetBase.GetRefSkeleton(),
+		bRecomputeNormals,
+		bRecomputeTangents);
+}
\ No newline at end of file
diff --git a/Engine/Plugins/ChaosClothAssetEditor/Source/ChaosClothAssetTools/Private/ChaosClothAsset/SkeletalMeshConverter.h b/Engine/Plugins/ChaosClothAssetEditor/Source/ChaosClothAssetTools/Private/ChaosClothAsset/SkeletalMeshConverter.h
new file mode 100644
--- /dev/null
+++ b/Engine/Plugins/ChaosClothAssetEditor/Source/ChaosClothAssetTools/Private/ChaosClothAsset/SkeletalMeshConverter.h
@@ -0,0 +1,16 @@
+// Copyright Epic Games, Inc. All Rights Reserved.
+
+#pragma once
+
+#include "ChaosClothAsset/SkeletalMeshConverterClassProvider.h"
+#include "SkeletalMeshConverter.generated.h"
+
+UCLASS()
+class UClothAssetEditorSkeletalMeshConverter : public UClothAssetSkeletalMeshConverter
+{
+	GENERATED_BODY()
+
+public:
+	/** Build a skeletal mesh from the specified Cloth Asset. */
+	virtual bool ExportToSkeletalMesh(const class UChaosClothAssetBase& ClothAssetBase, class USkeletalMesh& SkeletalMesh) const override;
+};
diff --git a/Engine/Plugins/ChaosClothAssetEditor/Source/ChaosClothAssetTools/Public/ChaosClothAsset/ClothPatternToDynamicMesh.h b/Engine/Plugins/ChaosClothAssetEditor/Source/ChaosClothAssetTools/Public/ChaosClothAsset/ClothPatternToDynamicMesh.h
--- a/Engine/Plugins/ChaosClothAssetEditor/Source/ChaosClothAssetTools/Public/ChaosClothAsset/ClothPatternToDynamicMesh.h
+++ b/Engine/Plugins/ChaosClothAssetEditor/Source/ChaosClothAssetTools/Public/ChaosClothAsset/ClothPatternToDynamicMesh.h
@@ -18,14 +18,28 @@ namespace UE::Chaos::ClothAsset
 * When PatternIndex = INDEX_NONE, convert the entire cloth asset to an FDynamicMesh3.
 * When VertexDataType = EClothPatternVertexType::Sim3D, this generates a welded mesh (with native welded indexing) when PatternIndex = INDEX_NONE,
 *	and an unwelded (with native 2D unwelded indexing) when PatternIndex is a valid SimPattern Index.
+* Use MaterialOffset to specify an offset into the material group ID, and run a Fix up on any triangles without valid material IDs.
+* When MaterialOffset is INDEX_NONE, no offset is applied and nbo fix up is run. 
 */
 class CHAOSCLOTHASSETTOOLS_API FClothPatternToDynamicMesh
 {
 public:
 
-	void Convert(const TSharedRef<const FManagedArrayCollection> ClothCollection, int32 PatternIndex, EClothPatternVertexType VertexDataType, UE::Geometry::FDynamicMesh3& MeshOut, bool bDisableAttributes = false);
-
-	void Convert(const UChaosClothAsset* ClothAssetMeshIn, int32 LODIndex, int32 PatternIndex, EClothPatternVertexType VertexDataType, UE::Geometry::FDynamicMesh3& MeshOut);
+	void Convert(
+		const TSharedRef<const FManagedArrayCollection> ClothCollection,
+		int32 PatternIndex,
+		EClothPatternVertexType VertexDataType,
+		UE::Geometry::FDynamicMesh3& MeshOut,
+		bool bDisableAttributes = false,
+		int32 MaterialOffset = INDEX_NONE);
+
+	void Convert(
+		const UChaosClothAsset* ClothAssetMeshIn,
+		int32 LODIndex,
+		int32 PatternIndex,
+		EClothPatternVertexType VertexDataType,
+		UE::Geometry::FDynamicMesh3& MeshOut,
+		int32 MaterialOffset = INDEX_NONE);
 };
 
 }	// namespace UE::Chaos::ClothAsset
diff --git a/Engine/Plugins/Experimental/ChaosOutfitAsset/Source/ChaosOutfitAssetDataflowNodes/Private/ChaosOutfitAsset/OutfitAssetTerminalNode.cpp b/Engine/Plugins/Experimental/ChaosOutfitAsset/Source/ChaosOutfitAssetDataflowNodes/Private/ChaosOutfitAsset/OutfitAssetTerminalNode.cpp
--- a/Engine/Plugins/Experimental/ChaosOutfitAsset/Source/ChaosOutfitAssetDataflowNodes/Private/ChaosOutfitAsset/OutfitAssetTerminalNode.cpp
+++ b/Engine/Plugins/Experimental/ChaosOutfitAsset/Source/ChaosOutfitAssetDataflowNodes/Private/ChaosOutfitAsset/OutfitAssetTerminalNode.cpp
@@ -1,22 +1,106 @@
 // Copyright Epic Games, Inc. All Rights Reserved.
 
 #include "ChaosOutfitAsset/OutfitAssetTerminalNode.h"
+#if WITH_EDITOR
+#include "AssetRegistry/AssetRegistryModule.h"
+#endif
 #include "ChaosOutfitAsset/OutfitAsset.h"
 #include "Dataflow/DataflowInputOutput.h"
+#if WITH_EDITOR
+#include "Dataflow/DataflowObjectInterface.h"
+#include "Engine/SkeletalMesh.h"
+#include "Misc/FileHelper.h"
+#include "Misc/ScopedSlowTask.h"
+#include "ContentBrowserModule.h"
+#include "IContentBrowserSingleton.h"
+#endif
 
 #include UE_INLINE_GENERATED_CPP_BY_NAME(OutfitAssetTerminalNode)
 
 #define LOCTEXT_NAMESPACE "ChaosOutfitAssetTerminalNode"
 
+namespace UE::Chaos::OutfitAsset::Private
+{
+#if WITH_EDITOR
+	static void ExportToSkeletalMesh(const TObjectPtr<const UChaosOutfitAsset>& OutfitAsset)
+	{
+		using namespace UE::Geometry;
+		if (OutfitAsset)
+		{
+			// Make a name from the outfit asset name
+			const FString SkeletalMeshPath = FPackageName::GetLongPackagePath(OutfitAsset->GetOutermost()->GetName());
+			const FString OutfitAssetName = OutfitAsset->GetName();
+			FString SkeletalMeshName = FString(TEXT("SK_")) + (OutfitAssetName.StartsWith(TEXT("OA_")) ? OutfitAssetName.RightChop(3) : OutfitAssetName);
+			FString SkeletalMeshPackageName = FPaths::Combine(SkeletalMeshPath, SkeletalMeshName);
+
+			// Open a Save As dialog
+			FSaveAssetDialogConfig SaveAssetDialogConfig;
+			SaveAssetDialogConfig.DefaultPath = SkeletalMeshPath;
+			SaveAssetDialogConfig.DefaultAssetName = SkeletalMeshName;
+			SaveAssetDialogConfig.AssetClassNames.Add(USkeletalMesh::StaticClass()->GetClassPathName());
+			SaveAssetDialogConfig.ExistingAssetPolicy = ESaveAssetDialogExistingAssetPolicy::AllowButWarn;
+			SaveAssetDialogConfig.DialogTitleOverride = LOCTEXT("SaveAssetDialogConfig", "Convert Outfit Asset To Skeletal Mesh As");
+
+			const FContentBrowserModule& ContentBrowserModule = FModuleManager::LoadModuleChecked<FContentBrowserModule>("ContentBrowser");
+			const FString AssetSavePath = ContentBrowserModule.Get().CreateModalSaveAssetDialog(SaveAssetDialogConfig);
+			const FString AssetSavePackageName = FPackageName::ObjectPathToPackageName(AssetSavePath);
+			FText OutError;
+			if (AssetSavePath.IsEmpty() || !FFileHelper::IsFilenameValidForSaving(AssetSavePackageName, OutError))
+			{
+				return;
+			}
+			SkeletalMeshPackageName = AssetSavePackageName;
+			SkeletalMeshName = FPackageName::GetLongPackageAssetName(SkeletalMeshPackageName);
+
+			constexpr float NumSteps = 0.f;
+			FScopedSlowTask SlowTask(NumSteps, LOCTEXT("ExportToSkeletalMesh", "Exporting Cloth Outfit to a Skeletal Mesh asset..."));
+			SlowTask.MakeDialog();
+
+			UPackage* const SkeletalMeshPackage = CreatePackage(*SkeletalMeshPackageName);
+			USkeletalMesh* const SkeletalMesh =
+				NewObject<USkeletalMesh>(
+					SkeletalMeshPackage,
+					USkeletalMesh::StaticClass(),
+					FName(SkeletalMeshName),
+					RF_Public | RF_Standalone | RF_Transactional);
+
+			// Export to skeletal mesh
+			if (SkeletalMesh)
+			{
+				OutfitAsset->ExportToSkeletalMesh(*SkeletalMesh);
+			}
+
+			SkeletalMesh->MarkPackageDirty();
+
+			// Notify the asset registry
+			FAssetRegistryModule::AssetCreated(SkeletalMesh);
+		}
+	}
+#endif  // #if WITH_EDITOR
+}
+
 FChaosOutfitAssetTerminalNode::FChaosOutfitAssetTerminalNode(const UE::Dataflow::FNodeParameters& InParam, FGuid InGuid)
 	: FDataflowTerminalNode(InParam, InGuid)
+#if WITH_EDITOR
+	, ConvertToSkeletalMesh(
+		FDataflowFunctionProperty::FDelegate::CreateLambda([this](UE::Dataflow::FContext& Context)
+			{
+				if (const UE::Dataflow::FEngineContext* EngineContext = Context.AsType<UE::Dataflow::FEngineContext>())
+				{
+					if (const UChaosOutfitAsset* const OwnerOutfitAsset = Cast<UChaosOutfitAsset>(EngineContext->Owner))
+					{
+						UE::Chaos::OutfitAsset::Private::ExportToSkeletalMesh(OwnerOutfitAsset);
+					}
+				}
+			}))
+#endif
 {
 	RegisterInputConnection(&Outfit);
 }
 
 void FChaosOutfitAssetTerminalNode::SetAssetValue(TObjectPtr<UObject> Asset, UE::Dataflow::FContext& Context) const
 {
-	if (UChaosOutfitAsset* const OutfitAsset = Cast<UChaosOutfitAsset>(Asset.Get()))
+	if (TObjectPtr<UChaosOutfitAsset> OutfitAsset = Cast<UChaosOutfitAsset>(Asset))
 	{
 		const TObjectPtr<const UChaosOutfit>& InOutfit = GetValue(Context, &Outfit);
 
diff --git a/Engine/Plugins/Experimental/ChaosOutfitAsset/Source/ChaosOutfitAssetDataflowNodes/Private/ChaosOutfitAsset/OutfitAssetTerminalNode.h b/Engine/Plugins/Experimental/ChaosOutfitAsset/Source/ChaosOutfitAssetDataflowNodes/Private/ChaosOutfitAsset/OutfitAssetTerminalNode.h
--- a/Engine/Plugins/Experimental/ChaosOutfitAsset/Source/ChaosOutfitAssetDataflowNodes/Private/ChaosOutfitAsset/OutfitAssetTerminalNode.h
+++ b/Engine/Plugins/Experimental/ChaosOutfitAsset/Source/ChaosOutfitAssetDataflowNodes/Private/ChaosOutfitAsset/OutfitAssetTerminalNode.h
@@ -2,6 +2,7 @@
 
 #pragma once 
 
+#include "Dataflow/DataflowFunctionProperty.h"
 #include "Dataflow/DataflowTerminalNode.h"
 #include "OutfitAssetTerminalNode.generated.h"
 
@@ -28,5 +29,9 @@ private:
 	/** Input outfit collection. */
 	UPROPERTY(Meta = (DataflowInput))
 	TObjectPtr<const UChaosOutfit> Outfit;
+
+	/** Export this outfit as a skeletal mesh asset. */
+	UPROPERTY(EditAnywhere, Category = "Outfit Asset Terminal", Meta = (ButtonImage = "Icons.Convert"))
+	FDataflowFunctionProperty ConvertToSkeletalMesh;
 };
 
