Chaos Outfit Asset - Added right click menu option to export a selected Outfit Asset to a Skeletal Mesh from the content browser.

#rnx

diff --git a/Engine/Plugins/Experimental/ChaosOutfitAsset/Source/ChaosOutfitAssetDataflowNodes/Private/ChaosOutfitAsset/OutfitAssetTerminalNode.cpp b/Engine/Plugins/Experimental/ChaosOutfitAsset/Source/ChaosOutfitAssetDataflowNodes/Private/ChaosOutfitAsset/OutfitAssetTerminalNode.cpp
--- a/Engine/Plugins/Experimental/ChaosOutfitAsset/Source/ChaosOutfitAssetDataflowNodes/Private/ChaosOutfitAsset/OutfitAssetTerminalNode.cpp
+++ b/Engine/Plugins/Experimental/ChaosOutfitAsset/Source/ChaosOutfitAssetDataflowNodes/Private/ChaosOutfitAsset/OutfitAssetTerminalNode.cpp
@@ -53,7 +53,7 @@ namespace UE::Chaos::OutfitAsset::Private
 			SkeletalMeshName = FPackageName::GetLongPackageAssetName(SkeletalMeshPackageName);
 
 			constexpr float NumSteps = 0.f;
-			FScopedSlowTask SlowTask(NumSteps, LOCTEXT("ExportToSkeletalMesh", "Exporting Cloth Outfit to a Skeletal Mesh asset..."));
+			FScopedSlowTask SlowTask(NumSteps, LOCTEXT("ConvertToSkeletalMesh", "Converting Cloth Outfit to a Skeletal Mesh asset..."));
 			SlowTask.MakeDialog();
 
 			UPackage* const SkeletalMeshPackage = CreatePackage(*SkeletalMeshPackageName);
diff --git a/Engine/Plugins/Experimental/ChaosOutfitAsset/Source/ChaosOutfitAssetDataflowNodes/Private/ChaosOutfitAsset/OutfitAssetTerminalNode.h b/Engine/Plugins/Experimental/ChaosOutfitAsset/Source/ChaosOutfitAssetDataflowNodes/Private/ChaosOutfitAsset/OutfitAssetTerminalNode.h
--- a/Engine/Plugins/Experimental/ChaosOutfitAsset/Source/ChaosOutfitAssetDataflowNodes/Private/ChaosOutfitAsset/OutfitAssetTerminalNode.h
+++ b/Engine/Plugins/Experimental/ChaosOutfitAsset/Source/ChaosOutfitAssetDataflowNodes/Private/ChaosOutfitAsset/OutfitAssetTerminalNode.h
@@ -30,7 +30,7 @@ private:
 	UPROPERTY(Meta = (DataflowInput))
 	TObjectPtr<const UChaosOutfit> Outfit;
 
-	/** Export this outfit as a skeletal mesh asset. */
+	/** Convert this outfit to a SkeletalMesh asset. */
 	UPROPERTY(EditAnywhere, Category = "Outfit Asset Terminal", Meta = (ButtonImage = "Icons.Convert"))
 	FDataflowFunctionProperty ConvertToSkeletalMesh;
 };
diff --git a/Engine/Plugins/Experimental/ChaosOutfitAsset/Source/ChaosOutfitAssetEditor/ChaosOutfitAssetEditor.Build.cs b/Engine/Plugins/Experimental/ChaosOutfitAsset/Source/ChaosOutfitAssetEditor/ChaosOutfitAssetEditor.Build.cs
--- a/Engine/Plugins/Experimental/ChaosOutfitAsset/Source/ChaosOutfitAssetEditor/ChaosOutfitAssetEditor.Build.cs
+++ b/Engine/Plugins/Experimental/ChaosOutfitAsset/Source/ChaosOutfitAssetEditor/ChaosOutfitAssetEditor.Build.cs
@@ -16,6 +16,7 @@ public class ChaosOutfitAssetEditor : ModuleRules
 				"ChaosClothAsset",
 				"ChaosClothAssetEngine",
 				"ChaosOutfitAssetEngine",
+				"ContentBrowser",
 				"Core",
 				"CoreUObject",
 				"DataflowEditor",
@@ -23,9 +24,10 @@ public class ChaosOutfitAssetEditor : ModuleRules
 				"Engine",
 				"InputCore",
 				"InteractiveToolsFramework",
-                "Projects",
-                "Slate",
+				"Projects",
+				"Slate",
 				"SlateCore",
+				"ToolMenus",
 				"UnrealEd"
 			}
 		);
diff --git a/Engine/Plugins/Experimental/ChaosOutfitAsset/Source/ChaosOutfitAssetEditor/Private/ChaosOutfitAsset/ChaosOutfitAssetEditorModule.cpp b/Engine/Plugins/Experimental/ChaosOutfitAsset/Source/ChaosOutfitAssetEditor/Private/ChaosOutfitAsset/ChaosOutfitAssetEditorModule.cpp
--- a/Engine/Plugins/Experimental/ChaosOutfitAsset/Source/ChaosOutfitAssetEditor/Private/ChaosOutfitAsset/ChaosOutfitAssetEditorModule.cpp
+++ b/Engine/Plugins/Experimental/ChaosOutfitAsset/Source/ChaosOutfitAssetEditor/Private/ChaosOutfitAsset/ChaosOutfitAssetEditorModule.cpp
@@ -1,18 +1,30 @@
 // Copyright Epic Games, Inc. All Rights Reserved.
 
 #include "ChaosOutfitAsset/AssetDefinition_OutfitAsset.h"
+#include "Algo/AllOf.h"
+#include "AssetRegistry/AssetRegistryModule.h"
+#include "ChaosOutfitAsset/OutfitAsset.h"
 #include "ChaosOutfitAsset/OutfitAssetEditorStyle.h"
 #include "ChaosOutfitAsset/OutfitEditorCommands.h"
+#include "Engine/SkeletalMesh.h"
+#include "Animation/Skeleton.h"
+#include "Factories/SkeletonFactory.h"
+#include "Framework/Commands/Commands.h"
+#include "Misc/ScopedSlowTask.h"
 #include "Modules/ModuleInterface.h"
 #include "Modules/ModuleManager.h"
+#include "AssetToolsModule.h"
+#include "ContentBrowserMenuContexts.h"
+#include "ToolMenus.h"
 
-#define LOCTEXT_NAMESPACE "OutfitAssetEngineModule"
+#define LOCTEXT_NAMESPACE "OutfitAssetEditorModule"
 
 namespace UE::Chaos::OutfitAsset
 {
 	class FOutfitAssetEditorModule : public IModuleInterface
 	{
 	public:
+		//~ Begin IModuleInterface interface
 		virtual void StartupModule() override
 		{
 			// Register asset icons
@@ -21,6 +33,10 @@ namespace UE::Chaos::OutfitAsset
 			// Register editor commands
 			FOutfitEditorCommands::Register();
 
+			// Register the asset menus
+			StartupCallbackDelegateHandle = UToolMenus::RegisterStartupCallback(
+				FSimpleMulticastDelegate::FDelegate::CreateRaw(this, &FOutfitAssetEditorModule::RegisterMenus));
+
 			// TODO Register visualizations (see FChaosClothAssetEditorModule::StartupModule() in ChaosClothAssetEditorModule.cpp)
 		}
 
@@ -29,9 +45,122 @@ namespace UE::Chaos::OutfitAsset
 			if (UObjectInitialized())
 			{
 				// TODO Deregister visualizations (see FChaosClothAssetEditorModule::StartupModule() in ChaosClothAssetEditorModule.cpp)
+
+				UToolMenus::UnRegisterStartupCallback(StartupCallbackDelegateHandle);
 				FOutfitEditorCommands::Unregister();
 			}
 		}
+		//~ End IModuleInterface interface
+
+private:
+		void RegisterMenus()
+		{
+			FToolMenuOwnerScoped OwnerScoped(this);  // Allows cleanup when module unloads.
+
+			UToolMenu* const ToolMenu = UToolMenus::Get()->ExtendMenu("ContentBrowser.AssetContextMenu.ChaosOutfitAsset");
+			FToolMenuSection& Section = ToolMenu->FindOrAddSection("GetAssetActions");
+
+			Section.AddDynamicEntry("ConvertToSkeletalMesh", FNewToolMenuSectionDelegate::CreateLambda(
+				[this](FToolMenuSection& Section)
+				{
+					if (const UContentBrowserAssetContextMenuContext* const Context = Section.FindContext<UContentBrowserAssetContextMenuContext>())
+					{
+						if (Algo::AllOf(
+							Context->SelectedAssets,  // Don't use Context->GetSelectedObjects() to avoid unnecessary loading the asset
+							[](const FAssetData& Asset)
+							{
+								return Asset.AssetClassPath == UChaosOutfitAsset::StaticClass()->GetClassPathName();
+							}))
+						{
+							const bool bIsChaosClothAssetToolsModuleLoaded = FModuleManager::Get().IsModuleLoaded("ChaosClothAssetTools");
+
+							const TSharedPtr<FUICommandList> CommandList = MakeShared<FUICommandList>();
+							CommandList->MapAction(
+								FOutfitEditorCommands::Get().GetConvertToSkeletalMesh(),
+								FExecuteAction::CreateWeakLambda(
+									Context, 
+									[Context]()  // ExecuteAction
+									{
+										const float NumSteps = (float)Context->SelectedAssets.Num() + 1.f;
+										FScopedSlowTask SlowTask(NumSteps, LOCTEXT("ConvertingToSkeletalMeshes", "Converting Outfit(s) to SkeletalMesh(es)..."));
+										SlowTask.MakeDialog();
+
+										SlowTask.EnterProgressFrame(1.f, LOCTEXT("LoadingOutfitAsset", "Loading OutfitAsset(s)..."));
+										const TArray<TObjectPtr<UObject>> SelectedObjects(Context->LoadSelectedObjects<UObject>());
+
+										for (const TObjectPtr<const UObject> SelectedObject : SelectedObjects)
+										{
+											SlowTask.EnterProgressFrame(1.f, LOCTEXT("ConvertingToSkeletalMesh", "Converting OutfitAsset to SkeletalMesh..."));
+											if (const TObjectPtr<const UChaosOutfitAsset> OutfitAsset = Cast<const UChaosOutfitAsset>(SelectedObject))
+											{
+												ConvertToSkeletalMesh(*OutfitAsset);
+											}
+										}
+									}),
+								FCanExecuteAction::CreateWeakLambda(
+									Context,
+									[bIsChaosClothAssetToolsModuleLoaded]()  // CanExecuteAction
+									{
+										return bIsChaosClothAssetToolsModuleLoaded;
+									}));
+
+							const TAttribute<FText> ToolTipOverride = bIsChaosClothAssetToolsModuleLoaded ?
+								LOCTEXT("ConvertToSkeletalMeshes", "Convert the selected OutfitAsset(s) to SkeletalMesh(es).") :
+								LOCTEXT("ChaosClothAssetEditorMustBeLoaded", "The ChaosClothAssetEditor plug-in must be loaded to enable this action.");
+							Section.AddMenuEntryWithCommandList(
+								FOutfitEditorCommands::Get().GetConvertToSkeletalMesh(),
+								CommandList,
+								TAttribute<FText>(),
+								ToolTipOverride,
+								FSlateIcon(FAppStyle::GetAppStyleSetName(), "Icons.Convert"));  // Could also use Icons.SkeletalMesh if it makes more sense
+						}
+					}
+				}));
+		}
+
+		static void ConvertToSkeletalMesh(const UChaosOutfitAsset& OutfitAsset)
+		{
+			// Make a unique name from the outfit asset name with the correct recommended asset prefixes
+			const FString SkeletalMeshPath = FPackageName::GetLongPackagePath(OutfitAsset.GetOutermost()->GetName());
+			const FString OutfitAssetName = OutfitAsset.GetName();
+			FString SkeletalMeshName = FString(TEXT("SK_")) + (OutfitAssetName.StartsWith(TEXT("OA_")) ? OutfitAssetName.RightChop(3) : OutfitAssetName);
+			FString SkeletalMeshPackageName = FPaths::Combine(SkeletalMeshPath, SkeletalMeshName);
+			if (FindPackage(nullptr, *SkeletalMeshPackageName))
+			{
+				MakeUniqueObjectName(nullptr, UPackage::StaticClass(), FName(SkeletalMeshPackageName)).ToString(SkeletalMeshPackageName);
+				SkeletalMeshName = FPaths::GetBaseFilename(SkeletalMeshPackageName);
+			}
+			UPackage* const SkeletalMeshPackage = CreatePackage(*SkeletalMeshPackageName);
+			if (USkeletalMesh* const SkeletalMesh =
+				NewObject<USkeletalMesh>(
+					SkeletalMeshPackage,
+					USkeletalMesh::StaticClass(),
+					FName(SkeletalMeshName),
+					RF_Public | RF_Standalone | RF_Transactional))
+			{
+				OutfitAsset.ExportToSkeletalMesh(*SkeletalMesh);
+
+				SkeletalMesh->MarkPackageDirty();
+
+				FAssetRegistryModule::AssetCreated(SkeletalMesh);  // Notify the asset registry
+
+				// Add the skeleton
+				FString SkeletonName = FString(TEXT("SKEL_")) + (OutfitAssetName.StartsWith(TEXT("OA_")) ? OutfitAssetName.RightChop(3) : OutfitAssetName);
+				FString SkeletonPackageName = FPaths::Combine(SkeletalMeshPath, SkeletonName);
+				if (FindPackage(nullptr, *SkeletonPackageName))
+				{
+					MakeUniqueObjectName(nullptr, UPackage::StaticClass(), FName(SkeletonPackageName)).ToString(SkeletonPackageName);
+					SkeletonName = FPaths::GetBaseFilename(SkeletonPackageName);
+				}
+				USkeletonFactory* const SkeletonFactory = NewObject<USkeletonFactory>();
+				SkeletonFactory->TargetSkeletalMesh = SkeletalMesh;
+
+				FAssetToolsModule& AssetToolsModule = FModuleManager::Get().LoadModuleChecked<FAssetToolsModule>("AssetTools");
+				AssetToolsModule.Get().CreateAsset(SkeletonName, SkeletalMeshPath, USkeleton::StaticClass(), SkeletonFactory);
+			}
+		}
+
+		FDelegateHandle StartupCallbackDelegateHandle;
 	};
 } // namespace UE::Chaos::OutfitAsset
 
diff --git a/Engine/Plugins/Experimental/ChaosOutfitAsset/Source/ChaosOutfitAssetEditor/Private/ChaosOutfitAsset/OutfitEditorCommands.cpp b/Engine/Plugins/Experimental/ChaosOutfitAsset/Source/ChaosOutfitAssetEditor/Private/ChaosOutfitAsset/OutfitEditorCommands.cpp
--- a/Engine/Plugins/Experimental/ChaosOutfitAsset/Source/ChaosOutfitAssetEditor/Private/ChaosOutfitAsset/OutfitEditorCommands.cpp
+++ b/Engine/Plugins/Experimental/ChaosOutfitAsset/Source/ChaosOutfitAssetEditor/Private/ChaosOutfitAsset/OutfitEditorCommands.cpp
@@ -20,7 +20,7 @@ namespace UE::Chaos::OutfitAsset
 	{
 		TBaseCharacterFXEditorCommands::RegisterCommands();
 
-		UI_COMMAND(OpenOutfitEditor, "Outfit Editor", "Open the Outfit Editor window", EUserInterfaceActionType::Button, FInputChord());
+		UI_COMMAND(ConvertToSkeletalMesh, "Convert to SkeletalMesh", "Convert the selected OutfitAsset(s) to SkeletalMesh(es).", EUserInterfaceActionType::Button, FInputChord());
 	}
 
 	void FOutfitEditorCommands::UpdateToolCommandBinding(UInteractiveTool* Tool, TSharedPtr<FUICommandList> UICommandList, bool bUnbind)
diff --git a/Engine/Plugins/Experimental/ChaosOutfitAsset/Source/ChaosOutfitAssetEditor/Private/ChaosOutfitAsset/OutfitEditorCommands.h b/Engine/Plugins/Experimental/ChaosOutfitAsset/Source/ChaosOutfitAssetEditor/Private/ChaosOutfitAsset/OutfitEditorCommands.h
--- a/Engine/Plugins/Experimental/ChaosOutfitAsset/Source/ChaosOutfitAssetEditor/Private/ChaosOutfitAsset/OutfitEditorCommands.h
+++ b/Engine/Plugins/Experimental/ChaosOutfitAsset/Source/ChaosOutfitAssetEditor/Private/ChaosOutfitAsset/OutfitEditorCommands.h
@@ -18,6 +18,11 @@ namespace UE::Chaos::OutfitAsset
 		 */
 		static void UpdateToolCommandBinding(UInteractiveTool* Tool, TSharedPtr<FUICommandList> UICommandList, bool bUnbind = false);
 
+		const TSharedPtr<FUICommandInfo>& GetConvertToSkeletalMesh() const
+		{
+			return ConvertToSkeletalMesh;
+		}
+
 	private:
 		friend TCommands<FOutfitEditorCommands>;
 
@@ -27,6 +32,6 @@ namespace UE::Chaos::OutfitAsset
 		virtual void GetToolDefaultObjectList(TArray<UInteractiveTool*>& ToolCDOs) override {}
 		//~ End TBaseCharacterFXEditorCommands<FOutfitEditorCommands> interface
 
-		TSharedPtr<FUICommandInfo> OpenOutfitEditor;
+		TSharedPtr<FUICommandInfo> ConvertToSkeletalMesh;
 	};
 } // namespace UE::Chaos::OutfitAsset
