Chaos Outfit Asset - Fixed crash in MergeBoneToSkeleton.
Chaos Cloth Asset - Fixed typo in comment.

#rnx

diff --git a/Engine/Plugins/ChaosClothAssetEditor/Source/ChaosClothAssetDataflowNodes/Private/ChaosClothAsset/USDImportNode_v2.cpp b/Engine/Plugins/ChaosClothAssetEditor/Source/ChaosClothAssetDataflowNodes/Private/ChaosClothAsset/USDImportNode_v2.cpp
--- a/Engine/Plugins/ChaosClothAssetEditor/Source/ChaosClothAssetDataflowNodes/Private/ChaosClothAsset/USDImportNode_v2.cpp
+++ b/Engine/Plugins/ChaosClothAssetEditor/Source/ChaosClothAssetDataflowNodes/Private/ChaosClothAsset/USDImportNode_v2.cpp
@@ -435,7 +435,7 @@ namespace UE::Chaos::ClothAsset::Private
 			{
 				const FUsdGeomSubset GeomSubset(SimMeshChildPrim);
 
-				// Read FamillyName
+				// Read FamilyName
 				const FUsdAttribute FamilyNameAttr = GeomSubset.GetFamilyNameAttr();
 				if (GetStringValue(FamilyNameAttr) != TEXT("sewing"))
 				{
diff --git a/Engine/Plugins/Experimental/ChaosOutfitAsset/Source/ChaosOutfitAssetEngine/Private/ChaosOutfitAsset/Outfit.cpp b/Engine/Plugins/Experimental/ChaosOutfitAsset/Source/ChaosOutfitAssetEngine/Private/ChaosOutfitAsset/Outfit.cpp
--- a/Engine/Plugins/Experimental/ChaosOutfitAsset/Source/ChaosOutfitAssetEngine/Private/ChaosOutfitAsset/Outfit.cpp
+++ b/Engine/Plugins/Experimental/ChaosOutfitAsset/Source/ChaosOutfitAssetEngine/Private/ChaosOutfitAsset/Outfit.cpp
@@ -136,7 +136,7 @@ namespace UE::Chaos::OutfitAsset::Private
 		else  // Doesn't already exists, add the bone
 		{
 			// Find its parent in the new skeleton
-			const FName ParentBoneName = InSkeleton.GetBoneName(ParentBoneIndex);
+			const FName ParentBoneName = (ParentBoneIndex == INDEX_NONE) ? NAME_None : InSkeleton.GetBoneName(ParentBoneIndex);
 			const int32 ExistingParentBoneIndex = InOutSkeleton.FindBoneIndex(ParentBoneName);
 			checkf(ParentBoneIndex == INDEX_NONE || ExistingParentBoneIndex != INDEX_NONE, TEXT("Either root or a valid parent should exist at this point."));
 			FReferenceSkeletonModifier ReferenceSkeletonModifier(InOutSkeleton, nullptr);
