Chaos Outfit Asset - Added HasValidBodySize() to the facade and a new OutfitQuery node to decide whether an outfit should be resized or not.

#rnx

#virtualized

diff --git a/Engine/Plugins/Experimental/ChaosOutfitAsset/Source/ChaosOutfitAssetDataflowNodes/Private/ChaosOutfitAsset/ChaosOutfitAssetDataflowNodesModule.cpp b/Engine/Plugins/Experimental/ChaosOutfitAsset/Source/ChaosOutfitAssetDataflowNodes/Private/ChaosOutfitAsset/ChaosOutfitAssetDataflowNodesModule.cpp
--- a/Engine/Plugins/Experimental/ChaosOutfitAsset/Source/ChaosOutfitAssetDataflowNodes/Private/ChaosOutfitAsset/ChaosOutfitAssetDataflowNodesModule.cpp
+++ b/Engine/Plugins/Experimental/ChaosOutfitAsset/Source/ChaosOutfitAssetDataflowNodes/Private/ChaosOutfitAsset/ChaosOutfitAssetDataflowNodesModule.cpp
@@ -11,6 +11,7 @@
 #include "ChaosOutfitAsset/MakeSizedOutfitNode.h"
 #include "ChaosOutfitAsset/MergeOutfitsNode.h"
 #include "ChaosOutfitAsset/OutfitAssetTerminalNode.h"
+#include "ChaosOutfitAsset/OutfitQueryNode.h"
 #include "ChaosOutfitAsset/SetOutfitClothCollectionNode.h"
 #include "ChaosOutfitAsset/SizedOutfitSourceAnyType.h"
 #include "Dataflow/DataflowAnyTypeRegistry.h"
@@ -44,6 +45,7 @@ namespace UE::Chaos::OutfitAsset
 			DATAFLOW_NODE_REGISTER_CREATION_FACTORY(FChaosOutfitAssetMergeOutfitsNode);
 			DATAFLOW_NODE_REGISTER_CREATION_FACTORY(FChaosSetOutfitClothCollectionNode);
 			DATAFLOW_NODE_REGISTER_CREATION_FACTORY(FChaosOutfitAssetTerminalNode);
+			DATAFLOW_NODE_REGISTER_CREATION_FACTORY(FChaosOutfitAssetOutfitQueryNode);
 PRAGMA_DISABLE_DEPRECATION_WARNINGS
 			// Move any deprecated node registrations under here
 PRAGMA_ENABLE_DEPRECATION_WARNINGS
diff --git a/Engine/Plugins/Experimental/ChaosOutfitAsset/Source/ChaosOutfitAssetDataflowNodes/Private/ChaosOutfitAsset/OutfitQueryNode.cpp b/Engine/Plugins/Experimental/ChaosOutfitAsset/Source/ChaosOutfitAssetDataflowNodes/Private/ChaosOutfitAsset/OutfitQueryNode.cpp
new file mode 100644
--- /dev/null
+++ b/Engine/Plugins/Experimental/ChaosOutfitAsset/Source/ChaosOutfitAssetDataflowNodes/Private/ChaosOutfitAsset/OutfitQueryNode.cpp
@@ -0,0 +1,56 @@
+// Copyright Epic Games, Inc. All Rights Reserved.
+
+#include "ChaosOutfitAsset/OutfitQueryNode.h"
+#include "ChaosOutfitAsset/CollectionOutfitFacade.h"
+
+#include UE_INLINE_GENERATED_CPP_BY_NAME(OutfitQueryNode)
+
+#define LOCTEXT_NAMESPACE "ChaosOutfitAssetOutfitQueryNode"
+
+FChaosOutfitAssetOutfitQueryNode::FChaosOutfitAssetOutfitQueryNode(const UE::Dataflow::FNodeParameters& InParam, FGuid InGuid)
+	: FDataflowNode(InParam, InGuid)
+{
+	RegisterInputConnection(&Outfit);
+	RegisterOutputConnection(&Outfit, &Outfit);
+	RegisterOutputConnection(&bHasAnyValidPieces);
+	RegisterOutputConnection(&bHasAnyValidBodySizes);
+}
+
+void FChaosOutfitAssetOutfitQueryNode::Evaluate(UE::Dataflow::FContext& Context, const FDataflowOutput* Out) const
+{
+	using namespace UE::Chaos::OutfitAsset;
+	if (Out->IsA(&Outfit))
+	{
+		SafeForwardInput(Context, &Outfit, &Outfit);
+	}
+	else if (Out->IsA<bool>(&bHasAnyValidPieces))
+	{
+		if (const TObjectPtr<const UChaosOutfit> InOutfit = GetValue(Context, &Outfit))
+		{
+			FCollectionOutfitConstFacade OutfitFacade(InOutfit->GetOutfitCollection());
+			if (OutfitFacade.IsValid())
+			{
+				const bool bOutHasAnyValidPieces = !OutfitFacade.GetOutfitPiecesGuids().IsEmpty();
+				SetValue(Context, bOutHasAnyValidPieces, &bHasAnyValidPieces);
+				return;
+			}
+		}
+		SetValue(Context, false, &bHasAnyValidPieces);
+	}
+	else if (Out->IsA<bool>(&bHasAnyValidBodySizes))
+	{
+		if (const TObjectPtr<const UChaosOutfit> InOutfit = GetValue(Context, &Outfit))
+		{
+			FCollectionOutfitConstFacade OutfitFacade(InOutfit->GetOutfitCollection());
+			if (OutfitFacade.IsValid())
+			{
+				const bool bOutHasAnyValidBodySizes = OutfitFacade.HasValidBodySize(bBodyPartMustExist, bMeasurementsMustExist, bInterpolationDataMustExist);
+				SetValue(Context, bOutHasAnyValidBodySizes, &bHasAnyValidBodySizes);
+				return;
+			}
+		}
+		SetValue(Context, false, &bHasAnyValidBodySizes);
+	}
+}
+
+#undef LOCTEXT_NAMESPACE
diff --git a/Engine/Plugins/Experimental/ChaosOutfitAsset/Source/ChaosOutfitAssetDataflowNodes/Private/ChaosOutfitAsset/OutfitQueryNode.h b/Engine/Plugins/Experimental/ChaosOutfitAsset/Source/ChaosOutfitAssetDataflowNodes/Private/ChaosOutfitAsset/OutfitQueryNode.h
new file mode 100644
--- /dev/null
+++ b/Engine/Plugins/Experimental/ChaosOutfitAsset/Source/ChaosOutfitAssetDataflowNodes/Private/ChaosOutfitAsset/OutfitQueryNode.h
@@ -0,0 +1,53 @@
+// Copyright Epic Games, Inc. All Rights Reserved.
+
+#pragma once 
+
+#include "Dataflow/DataflowNode.h"
+#include "ChaosOutfitAsset/Outfit.h"
+#include "OutfitQueryNode.generated.h"
+
+
+class UDynamicMesh;
+class UMaterialInterface;
+
+/**
+ * Query an Outfit about its properties.
+ */
+USTRUCT(Meta = (DataflowOutfit))
+struct FChaosOutfitAssetOutfitQueryNode : public FDataflowNode
+{
+	GENERATED_USTRUCT_BODY()
+	DATAFLOW_NODE_DEFINE_INTERNAL(FChaosOutfitAssetOutfitQueryNode, "OutfitQuery", "Outfit", "Outfit Query")
+
+public:
+
+	FChaosOutfitAssetOutfitQueryNode(const UE::Dataflow::FNodeParameters& InParam, FGuid InGuid = FGuid::NewGuid());
+
+	virtual void Evaluate(UE::Dataflow::FContext& Context, const FDataflowOutput* Out) const override;
+
+private:
+
+	/** Input/output collection (Output is always a passthrough)*/
+	UPROPERTY(Meta = (DataflowInput, DataflowOutput, DataflowPassthrough = "Outfit"))
+	TObjectPtr<const UChaosOutfit> Outfit;
+
+	/** Has this outfit any valid pieces? */
+	UPROPERTY(Meta = (DataflowOutput))
+	bool bHasAnyValidPieces = false;
+
+	/** Has this outfit any valid body sizes? */
+	UPROPERTY(Meta = (DataflowOutput))
+	bool bHasAnyValidBodySizes = false;
+
+	/** Check that body parts are present in the asset registry when checking for valid body sizes. */
+	UPROPERTY(EditAnywhere, Category = "Outfit Query")
+	bool bBodyPartMustExist = false;
+	
+	/** Check that measurements are valid when checking for valid body sizes. */
+	UPROPERTY(EditAnywhere, Category = "Outfit Query")
+	bool bMeasurementsMustExist = true;
+	
+	/** Check that some interpolation data exists when checking for valid body sizes. */
+	UPROPERTY(EditAnywhere, Category = "Outfit Query")
+	bool bInterpolationDataMustExist = true;
+};
diff --git a/Engine/Plugins/Experimental/ChaosOutfitAsset/Source/ChaosOutfitAssetEngine/Private/ChaosOutfitAsset/CollectionOutfitFacade.cpp b/Engine/Plugins/Experimental/ChaosOutfitAsset/Source/ChaosOutfitAssetEngine/Private/ChaosOutfitAsset/CollectionOutfitFacade.cpp
--- a/Engine/Plugins/Experimental/ChaosOutfitAsset/Source/ChaosOutfitAssetEngine/Private/ChaosOutfitAsset/CollectionOutfitFacade.cpp
+++ b/Engine/Plugins/Experimental/ChaosOutfitAsset/Source/ChaosOutfitAssetEngine/Private/ChaosOutfitAsset/CollectionOutfitFacade.cpp
@@ -1,6 +1,10 @@
 // Copyright Epic Games, Inc. All Rights Reserved.
 
 #include "ChaosOutfitAsset/CollectionOutfitFacade.h"
+#include "Algo/AllOf.h"
+#include "Algo/AnyOf.h"
+#include "AssetRegistry/AssetData.h"
+#include "AssetRegistry/AssetRegistryModule.h"
 #include "ChaosClothAsset/ClothAssetBase.h"
 #include "ChaosOutfitAsset/BodyUserData.h"
 #include "ChaosOutfitAsset/OutfitAssetPrivate.h"
@@ -382,6 +386,54 @@ namespace UE::Chaos::OutfitAsset
 		return OutfitGuids;
 	}
 
+	bool FCollectionOutfitConstFacade::HasValidBodySize(bool bBodyPartMustExist, bool bBodyMeasurementsMustExist, bool bInterpolationDataMustExist) const
+	{
+		const FAssetRegistryModule& AssetRegistryModule = FModuleManager::LoadModuleChecked<FAssetRegistryModule>("AssetRegistry");
+
+		for (int32 BodySize = 0; BodySize < GetNumBodySizes(); ++BodySize)
+		{
+			const TConstArrayView<FString> BodyPartsSkeletalMeshes = GetBodySizeBodyPartsSkeletalMeshes(BodySize);
+			if (!Algo::AnyOf(BodyPartsSkeletalMeshes, [&AssetRegistryModule, bBodyPartMustExist](const FString& BodyPartsSkeletalMesh)
+				{
+					if (bBodyPartMustExist)
+					{
+						const FAssetData AssetData = AssetRegistryModule.Get().GetAssetByObjectPath(BodyPartsSkeletalMesh);
+						return AssetData.IsValid();  //  Body part exists
+					}
+					return !BodyPartsSkeletalMesh.IsEmpty();  // Body part is not empty
+				}))
+			{
+				continue;  // Invalid body parts, try another size
+			}
+
+			if (bBodyMeasurementsMustExist)
+			{
+				const TMap<FString, float> BodySizeMeasurements = GetBodySizeMeasurements(BodySize);
+				if (BodySizeMeasurements.IsEmpty() || Algo::AllOf(BodySizeMeasurements, [](const TPair<FString, float>& Measurement)
+					{
+						return Measurement.Value == UChaosOutfitAssetBodyUserData::InvalidMeasurement;
+					}))
+				{
+					continue;  // No measurement, or all measurement values are 0, try another size
+				}
+			}
+
+			if (bInterpolationDataMustExist)
+			{
+				const FRBFInterpolationDataWrapper RBFInterpolationDataWrapper = GetBodySizeInterpolationData(BodySize);
+				if (!RBFInterpolationDataWrapper.SampleIndices.Num() ||
+					!RBFInterpolationDataWrapper.SampleRestPositions.Num() ||
+					!RBFInterpolationDataWrapper.InterpolationWeights.Num())
+				{
+					continue;  // Needs RBF data, try another size
+				}
+			}
+
+			return true;  // Found at least one good size
+		}
+		return false;
+	}
+
 	TArray<FString> FCollectionOutfitConstFacade::GetOutfitBodyPartsSkeletalMeshes() const
 	{
 		check(IsValid());
diff --git a/Engine/Plugins/Experimental/ChaosOutfitAsset/Source/ChaosOutfitAssetEngine/Public/ChaosOutfitAsset/CollectionOutfitFacade.h b/Engine/Plugins/Experimental/ChaosOutfitAsset/Source/ChaosOutfitAssetEngine/Public/ChaosOutfitAsset/CollectionOutfitFacade.h
--- a/Engine/Plugins/Experimental/ChaosOutfitAsset/Source/ChaosOutfitAssetEngine/Public/ChaosOutfitAsset/CollectionOutfitFacade.h
+++ b/Engine/Plugins/Experimental/ChaosOutfitAsset/Source/ChaosOutfitAssetEngine/Public/ChaosOutfitAsset/CollectionOutfitFacade.h
@@ -49,6 +49,9 @@ namespace UE::Chaos::OutfitAsset
 		/** Return the sized outfits GUID, one per sized outfit which can contain different sized outfits. */
 		CHAOSOUTFITASSETENGINE_API TArray<FGuid> GetOutfitGuids() const;
 
+		/** Return whether this outfit has at least one valid body size. Definition of valid depends on the passed arguments. */
+		CHAOSOUTFITASSETENGINE_API bool HasValidBodySize(bool bBodyPartMustExist = false, bool bMeasurementsMustExist = false, bool bInterpolationDataMustExist = false) const;
+
 		/** Return all unique body part skeletal meshes used by this outfit. */
 		CHAOSOUTFITASSETENGINE_API TArray<FString> GetOutfitBodyPartsSkeletalMeshes() const;
 
