Chaos Cloth/Outfit Asset - Adjusted module dependencies.

#rnx

diff --git a/Engine/Plugins/ChaosClothAsset/Source/ChaosClothAssetEngine/ChaosClothAssetEngine.Build.cs b/Engine/Plugins/ChaosClothAsset/Source/ChaosClothAssetEngine/ChaosClothAssetEngine.Build.cs
--- a/Engine/Plugins/ChaosClothAsset/Source/ChaosClothAssetEngine/ChaosClothAssetEngine.Build.cs
+++ b/Engine/Plugins/ChaosClothAsset/Source/ChaosClothAssetEngine/ChaosClothAssetEngine.Build.cs
@@ -13,7 +13,8 @@ public class ChaosClothAssetEngine : ModuleRules
 			{
 				"DataflowCore",
 				"DataflowEngine",
-				"DataflowSimulation"
+				"DataflowSimulation",
+				"ClothingSystemRuntimeCommon"
 			}
 		);
 
@@ -29,7 +30,6 @@ public class ChaosClothAssetEngine : ModuleRules
 				"ChaosClothAsset",
 				"ChaosCloth",
 				"ChaosCaching",
-				"ClothingSystemRuntimeCommon",
 				"ClothingSystemRuntimeInterface"
 			}
 		);
diff --git a/Engine/Plugins/Experimental/ChaosOutfitAsset/Source/ChaosOutfitAssetEngine/ChaosOutfitAssetEngine.Build.cs b/Engine/Plugins/Experimental/ChaosOutfitAsset/Source/ChaosOutfitAssetEngine/ChaosOutfitAssetEngine.Build.cs
--- a/Engine/Plugins/Experimental/ChaosOutfitAsset/Source/ChaosOutfitAssetEngine/ChaosOutfitAssetEngine.Build.cs
+++ b/Engine/Plugins/Experimental/ChaosOutfitAsset/Source/ChaosOutfitAssetEngine/ChaosOutfitAssetEngine.Build.cs
@@ -11,8 +11,7 @@ public class ChaosOutfitAssetEngine : ModuleRules
 		PublicDependencyModuleNames.AddRange(
 			new string[]
 			{
-				"DataflowCore",
-				"DataflowEngine",
+				"ChaosClothAssetEngine",
 			}
 		);
 
@@ -23,8 +22,7 @@ public class ChaosOutfitAssetEngine : ModuleRules
 				"CoreUObject",
 				"Chaos",
 				"ChaosClothAsset",
-				"ChaosClothAssetEngine",
-				"ClothingSystemRuntimeCommon",
+				"DataflowEngine",
 				"Engine",
 				"MeshResizingCore",
 				"RenderCore",
