Chaos Cloth/Outfit Asset - Tidy up API declarations.

#rnx

diff --git a/Engine/Plugins/Experimental/ChaosOutfitAsset/Source/ChaosOutfitAssetEngine/Public/ChaosOutfitAsset/BodyUserData.h b/Engine/Plugins/Experimental/ChaosOutfitAsset/Source/ChaosOutfitAssetEngine/Public/ChaosOutfitAsset/BodyUserData.h
--- a/Engine/Plugins/Experimental/ChaosOutfitAsset/Source/ChaosOutfitAssetEngine/Public/ChaosOutfitAsset/BodyUserData.h
+++ b/Engine/Plugins/Experimental/ChaosOutfitAsset/Source/ChaosOutfitAssetEngine/Public/ChaosOutfitAsset/BodyUserData.h
@@ -8,8 +8,8 @@
 /**
  * Asset user data attached to the merged body and face skeletal mesh.
  */
-UCLASS()
-class CHAOSOUTFITASSETENGINE_API UChaosOutfitAssetBodyUserData : public UAssetUserData
+UCLASS(MinimalAPI)
+class UChaosOutfitAssetBodyUserData : public UAssetUserData
 {
 	GENERATED_BODY()
 
diff --git a/Engine/Plugins/Experimental/ChaosOutfitAsset/Source/ChaosOutfitAssetEngine/Public/ChaosOutfitAsset/CollectionOutfitFacade.h b/Engine/Plugins/Experimental/ChaosOutfitAsset/Source/ChaosOutfitAssetEngine/Public/ChaosOutfitAsset/CollectionOutfitFacade.h
--- a/Engine/Plugins/Experimental/ChaosOutfitAsset/Source/ChaosOutfitAssetEngine/Public/ChaosOutfitAsset/CollectionOutfitFacade.h
+++ b/Engine/Plugins/Experimental/ChaosOutfitAsset/Source/ChaosOutfitAssetEngine/Public/ChaosOutfitAsset/CollectionOutfitFacade.h
@@ -6,6 +6,8 @@
 #include "Templates/UniquePtr.h"
 #include "ChaosOutfitAsset/OutfitCollection.h"
 
+#define UE_API CHAOSOUTFITASSETENGINE_API
+
 template<typename T> class TManagedArray;
 class UChaosClothAssetBase;
 class USkeletalMesh;
@@ -26,49 +28,49 @@ namespace UE::Chaos::OutfitAsset
 	class FCollectionOutfitConstFacade
 	{
 	public:
-		struct CHAOSOUTFITASSETENGINE_API FRBFInterpolationDataWrapper
+		struct FRBFInterpolationDataWrapper
 		{
 			TConstArrayView<TArray<int32>> SampleIndices;
 			TConstArrayView<TArray<FVector3f>> SampleRestPositions;
 			TConstArrayView<TArray<float>> InterpolationWeights;
 		};
 
-		CHAOSOUTFITASSETENGINE_API explicit FCollectionOutfitConstFacade(const TSharedRef<const FManagedArrayCollection>& InManagedArrayCollection);
-		CHAOSOUTFITASSETENGINE_API explicit FCollectionOutfitConstFacade(const FManagedArrayCollection& InManagedArrayCollection);
+		UE_API explicit FCollectionOutfitConstFacade(const TSharedRef<const FManagedArrayCollection>& InManagedArrayCollection);
+		UE_API explicit FCollectionOutfitConstFacade(const FManagedArrayCollection& InManagedArrayCollection);
 
 		FCollectionOutfitConstFacade() = delete;
 		FCollectionOutfitConstFacade(const FCollectionOutfitConstFacade&) = delete;
 		FCollectionOutfitConstFacade& operator=(const FCollectionOutfitConstFacade&) = delete;
 		FCollectionOutfitConstFacade(FCollectionOutfitConstFacade&&) = default;
 		FCollectionOutfitConstFacade& operator=(FCollectionOutfitConstFacade&&) = default;
-		CHAOSOUTFITASSETENGINE_API virtual ~FCollectionOutfitConstFacade();
+		UE_API virtual ~FCollectionOutfitConstFacade();
 
 		/** Return whether the facade is defined on the collection. */
-		CHAOSOUTFITASSETENGINE_API bool IsValid() const;
+		UE_API bool IsValid() const;
 
 		/** Return the sized outfits GUID, one per sized outfit which can contain different sized outfits. */
-		CHAOSOUTFITASSETENGINE_API TArray<FGuid> GetOutfitGuids() const;
+		UE_API TArray<FGuid> GetOutfitGuids() const;
 
 		/** Return whether this outfit has at least one valid body size. Definition of valid depends on the passed arguments. */
-		CHAOSOUTFITASSETENGINE_API bool HasValidBodySize(bool bBodyPartMustExist = false, bool bMeasurementsMustExist = false, bool bInterpolationDataMustExist = false) const;
+		UE_API bool HasValidBodySize(bool bBodyPartMustExist = false, bool bMeasurementsMustExist = false, bool bInterpolationDataMustExist = false) const;
 
 		/** Return all unique body part skeletal meshes used by this outfit. */
-		CHAOSOUTFITASSETENGINE_API TArray<FString> GetOutfitBodyPartsSkeletalMeshes() const;
+		UE_API TArray<FString> GetOutfitBodyPartsSkeletalMeshes() const;
 
 		/** Return the body size available for the specified outfit GUID. */
-		CHAOSOUTFITASSETENGINE_API TArray<int32> GetOutfitBodySizes(const FGuid& Guid) const;
+		UE_API TArray<int32> GetOutfitBodySizes(const FGuid& Guid) const;
 
 		/** Return the name of the original asset making up the specified outfit body size. */
-		CHAOSOUTFITASSETENGINE_API const FString& GetOutfitName(const FGuid& Guid, const int32 BodySize) const;
+		UE_API const FString& GetOutfitName(const FGuid& Guid, const int32 BodySize) const;
 
 		/** Return all outfit pieces asset GUIDs. */
-		CHAOSOUTFITASSETENGINE_API TConstArrayView<FGuid> GetOutfitPiecesGuids() const;
+		UE_API TConstArrayView<FGuid> GetOutfitPiecesGuids() const;
 
 		/** Return the outfit pieces making up the specified outfit body size. */
-		CHAOSOUTFITASSETENGINE_API TMap<FGuid, FString> GetOutfitPieces(const FGuid& Guid, const int32 BodySize) const;
+		UE_API TMap<FGuid, FString> GetOutfitPieces(const FGuid& Guid, const int32 BodySize) const;
 
 		/** Return the number of body sizes. */
-		CHAOSOUTFITASSETENGINE_API int32 GetNumBodySizes() const;
+		UE_API int32 GetNumBodySizes() const;
 
 		/** Return whether this body size already exist. */
 		bool HasBodySize(const FString& Name) const
@@ -77,31 +79,31 @@ namespace UE::Chaos::OutfitAsset
 		}
 
 		/** Return the body size index for the specified name, or INDEX_NONE if it doesn't exist in this outfit. */
-		CHAOSOUTFITASSETENGINE_API int32 FindBodySize(const FString& Name) const;
+		UE_API int32 FindBodySize(const FString& Name) const;
 
 		/** Return the closest size to the specified measurements, or INDEX_NONE if this outfit has no sizes or the measurements are incomplete. */
-		CHAOSOUTFITASSETENGINE_API int32 FindClosestBodySize(const TMap<FString, float>& Measurements) const;
+		UE_API int32 FindClosestBodySize(const TMap<FString, float>& Measurements) const;
 
 		/** Return the closest size to the specified body, or INDEX_NONE if this outfit has no sizes or the body has no measurement data. */
-		CHAOSOUTFITASSETENGINE_API int32 FindClosestBodySize(const USkeletalMesh& BodyPart) const;
+		UE_API int32 FindClosestBodySize(const USkeletalMesh& BodyPart) const;
 
 		/** Return the specified body size name which is unique. */
-		CHAOSOUTFITASSETENGINE_API const FString& GetBodySizeName(const int32 BodySize) const;
+		UE_API const FString& GetBodySizeName(const int32 BodySize) const;
 
 		/** Return the body part skeletal meshes path name for the specified body size. */
-		CHAOSOUTFITASSETENGINE_API TConstArrayView<FString> GetBodySizeBodyPartsSkeletalMeshes(const int32 BodySize) const;
+		UE_API TConstArrayView<FString> GetBodySizeBodyPartsSkeletalMeshes(const int32 BodySize) const;
 
 		/** Return the body size body part offset (to index into the global body part array). Returns INDEX_NONE if BodySize is invalid.*/
-		CHAOSOUTFITASSETENGINE_API int32 GetBodySizeBodyPartOffset(const int32 BodySize) const;
+		UE_API int32 GetBodySizeBodyPartOffset(const int32 BodySize) const;
 
 		/** Return the body size body part count. Returns 0 if BodySize is invalid.*/
-		CHAOSOUTFITASSETENGINE_API int32 GetBodySizeBodyPartCount(const int32 BodySize) const;
+		UE_API int32 GetBodySizeBodyPartCount(const int32 BodySize) const;
 
 		/** Return  for the specified body size. */
-		CHAOSOUTFITASSETENGINE_API FRBFInterpolationDataWrapper GetBodySizeInterpolationData(const int32 BodySize) const;
+		UE_API FRBFInterpolationDataWrapper GetBodySizeInterpolationData(const int32 BodySize) const;
 
 		/** Return the body measurements stored for this specified body size. */
-		CHAOSOUTFITASSETENGINE_API TMap<FString, float> GetBodySizeMeasurements(const int32 BodySize) const;
+		UE_API TMap<FString, float> GetBodySizeMeasurements(const int32 BodySize) const;
 
 
 	protected:
@@ -118,27 +120,27 @@ namespace UE::Chaos::OutfitAsset
 	class FCollectionOutfitFacade final : public FCollectionOutfitConstFacade
 	{
 	public:
-		CHAOSOUTFITASSETENGINE_API explicit FCollectionOutfitFacade(const TSharedRef<FManagedArrayCollection>& InManagedArrayCollection);
-		CHAOSOUTFITASSETENGINE_API explicit FCollectionOutfitFacade(FManagedArrayCollection& InManagedArrayCollection);
+		UE_API explicit FCollectionOutfitFacade(const TSharedRef<FManagedArrayCollection>& InManagedArrayCollection);
+		UE_API explicit FCollectionOutfitFacade(FManagedArrayCollection& InManagedArrayCollection);
 	
 		FCollectionOutfitFacade() = delete;
 		FCollectionOutfitFacade(const FCollectionOutfitFacade&) = delete;
 		FCollectionOutfitFacade& operator=(const FCollectionOutfitFacade&) = delete;
 		FCollectionOutfitFacade(FCollectionOutfitFacade&&) = default;
 		FCollectionOutfitFacade& operator=(FCollectionOutfitFacade&&) = default;
-		CHAOSOUTFITASSETENGINE_API virtual ~FCollectionOutfitFacade() override;
+		UE_API virtual ~FCollectionOutfitFacade() override;
 
 		/** Add the Selection attributes to the underlying collection. */
-		CHAOSOUTFITASSETENGINE_API void DefineSchema();
+		UE_API void DefineSchema();
 
 		/** Remove all outfit and size information from this collection. */
-		CHAOSOUTFITASSETENGINE_API void Reset();
+		UE_API void Reset();
 
 		/**
 		 * Add a new body size and returns its index.
 		 * Any pre-existing sizes data will be replaced.
 		 */
-		CHAOSOUTFITASSETENGINE_API int32 AddBodySize(
+		UE_API int32 AddBodySize(
 			const FString& Name,
 			const TConstArrayView<FString> BodyPartsSkeletalMeshes,
 			const TMap<FString, float>& Measurements,
@@ -156,7 +158,7 @@ namespace UE::Chaos::OutfitAsset
 		 * Add a new sized outfit under the specified outfit GUID.
 		 * Any pre-existing sizes data will be replaced.
 		 */
-		CHAOSOUTFITASSETENGINE_API void AddOutfit(
+		UE_API void AddOutfit(
 			const FGuid& Guid,
 			const int32 BodySize,
 			const UChaosClothAssetBase& ClothAssetBase);
@@ -167,7 +169,7 @@ namespace UE::Chaos::OutfitAsset
 		 * @param Other The source outfit collection to get the data from.
 		 * @param SelectedBodySize The body size of the source outfit to append, or INDEX_NONE to append all sizes.
 		 */
-		CHAOSOUTFITASSETENGINE_API void Append(const FCollectionOutfitConstFacade& Other, int32 SelectedBodySize = INDEX_NONE);
+		UE_API void Append(const FCollectionOutfitConstFacade& Other, int32 SelectedBodySize = INDEX_NONE);
 
 		/**
 		 * Append an existing outfit facade to this collection.
@@ -189,3 +191,5 @@ namespace UE::Chaos::OutfitAsset
 		Private::FOutfitCollection* GetOutfitCollection() { return const_cast<Private::FOutfitCollection*>(OutfitCollection.Get()); }
 	};
 }
+
+#undef UE_API
diff --git a/Engine/Plugins/Experimental/ChaosOutfitAsset/Source/ChaosOutfitAssetEngine/Public/ChaosOutfitAsset/Outfit.h b/Engine/Plugins/Experimental/ChaosOutfitAsset/Source/ChaosOutfitAssetEngine/Public/ChaosOutfitAsset/Outfit.h
--- a/Engine/Plugins/Experimental/ChaosOutfitAsset/Source/ChaosOutfitAssetEngine/Public/ChaosOutfitAsset/Outfit.h
+++ b/Engine/Plugins/Experimental/ChaosOutfitAsset/Source/ChaosOutfitAssetEngine/Public/ChaosOutfitAsset/Outfit.h
@@ -10,6 +10,8 @@
 
 #include "Outfit.generated.h"
 
+#define UE_API CHAOSOUTFITASSETENGINE_API
+
 class FSkeletalMeshRenderData;
 class UChaosClothAssetBase;
 class UPhysicsAsset;
@@ -20,7 +22,7 @@ struct FManagedArrayCollection;
 struct FSkeletalMaterial;
 
 USTRUCT()
-struct CHAOSOUTFITASSETENGINE_API FChaosOutfitPiece final
+struct FChaosOutfitPiece final
 {
 	GENERATED_BODY()
 
@@ -37,26 +39,26 @@ struct CHAOSOUTFITASSETENGINE_API FChaosOutfitPiece final
 
 	TArray<TSharedRef<const FManagedArrayCollection>> Collections;
 
-	FChaosOutfitPiece();
-	~FChaosOutfitPiece();
-	FChaosOutfitPiece(
+	UE_API FChaosOutfitPiece();
+	UE_API ~FChaosOutfitPiece();
+	UE_API FChaosOutfitPiece(
 		FName InName, 
 		FGuid InAssetGuid, 
 		const UPhysicsAsset* InPhysicsAsset, 
 		const FChaosClothSimulationModel& InClothSimulationModel, 
 		const TArray<TSharedRef<const FManagedArrayCollection>>& InCollections);
 
-	FChaosOutfitPiece(const UChaosClothAssetBase& ClothAssetBase, int32 ModelIndex);
+	UE_API FChaosOutfitPiece(const UChaosClothAssetBase& ClothAssetBase, int32 ModelIndex);
 	FChaosOutfitPiece(const TArray<FChaosOutfitPiece>& Pieces, int32 ModelIndex) : FChaosOutfitPiece(Pieces[ModelIndex]) {}
 
-	FChaosOutfitPiece(const FChaosOutfitPiece& Other);
+	UE_API FChaosOutfitPiece(const FChaosOutfitPiece& Other);
 
-	FChaosOutfitPiece(FChaosOutfitPiece&& Other);
+	UE_API FChaosOutfitPiece(FChaosOutfitPiece&& Other);
 
-	FChaosOutfitPiece& operator=(const FChaosOutfitPiece& Other);
-	FChaosOutfitPiece& operator=(FChaosOutfitPiece&& Other);
+	UE_API FChaosOutfitPiece& operator=(const FChaosOutfitPiece& Other);
+	UE_API FChaosOutfitPiece& operator=(FChaosOutfitPiece&& Other);
 
-	bool Serialize(FArchive& Ar);
+	UE_API bool Serialize(FArchive& Ar);
 
 private:
 	friend UChaosOutfit;
@@ -78,22 +80,22 @@ struct TStructOpsTypeTraits<FChaosOutfitPiece> : public TStructOpsTypeTraitsBase
 /**
  * Outfit class to handle the assembly of an outfit asset.
  */
-UCLASS(Experimental, BlueprintType)
-class CHAOSOUTFITASSETENGINE_API UChaosOutfit final : public UObject
+UCLASS(Experimental, MinimalAPI, BlueprintType)
+class UChaosOutfit final : public UObject
 {
 	GENERATED_BODY()
 
 public:
-	UChaosOutfit(const FObjectInitializer& ObjectInitializer);
-	UChaosOutfit(FVTableHelper& Helper);  // This is declared so we can use TUniquePtr<FClothSimulationModel> with just a forward declare of that class
-	virtual ~UChaosOutfit() override;
+	UE_API UChaosOutfit(const FObjectInitializer& ObjectInitializer);
+	UE_API UChaosOutfit(FVTableHelper& Helper);  // This is declared so we can use TUniquePtr<FClothSimulationModel> with just a forward declare of that class
+	UE_API virtual ~UChaosOutfit() override;
 
-	void Append(const UChaosOutfit& Other, const FString& BodySizeNameFilter = {});
+	UE_API void Append(const UChaosOutfit& Other, const FString& BodySizeNameFilter = {});
 
-	void Add(const UChaosClothAssetBase& ClothAssetBase);
-	void Add(const FChaosSizedOutfitSource& SizedOutfitSource, const FGuid& OutfitGuid = FGuid::NewGuid());
+	UE_API void Add(const UChaosClothAssetBase& ClothAssetBase);
+	UE_API void Add(const FChaosSizedOutfitSource& SizedOutfitSource, const FGuid& OutfitGuid = FGuid::NewGuid());
 
-	void CopyTo(
+	UE_API void CopyTo(
 		TArray<FChaosOutfitPiece>& OutPieces,
 		FReferenceSkeleton& OutReferenceSkeleton,
 		TUniquePtr<FSkeletalMeshRenderData>& OutSkeletalMeshRenderData,
@@ -120,20 +122,20 @@ public:
 	}
 
 	/** Return the number of LODs (max LOD contained in any piece). */
-	int32 GetNumLods() const;
+	UE_API int32 GetNumLods() const;
 
 	/** Return the outfit pieces cloth collections for all LODs. */
-	TArray<TSharedRef<const FManagedArrayCollection>> GetClothCollections(int32 LodIndex) const;
+	UE_API TArray<TSharedRef<const FManagedArrayCollection>> GetClothCollections(int32 LodIndex) const;
 
 	const FReferenceSkeleton& GetReferenceSkeleton() const
 	{
 		return ReferenceSkeleton;
 	}
 
-	bool HasBodySize(const FString& SizeName) const;
+	UE_API bool HasBodySize(const FString& SizeName) const;
 
 	//~ Begin UObject interface
-	virtual void Serialize(FArchive& Ar) override;
+	UE_API virtual void Serialize(FArchive& Ar) override;
 	//~ End UObject interface
 
 	struct FLODRenderData;
@@ -184,3 +186,5 @@ private:
 
 	TUniquePtr<FRenderData> RenderData;
 };
+
+#undef UE_API
diff --git a/Engine/Plugins/Experimental/ChaosOutfitAsset/Source/ChaosOutfitAssetEngine/Public/ChaosOutfitAsset/OutfitCollection.h b/Engine/Plugins/Experimental/ChaosOutfitAsset/Source/ChaosOutfitAssetEngine/Public/ChaosOutfitAsset/OutfitCollection.h
--- a/Engine/Plugins/Experimental/ChaosOutfitAsset/Source/ChaosOutfitAssetEngine/Public/ChaosOutfitAsset/OutfitCollection.h
+++ b/Engine/Plugins/Experimental/ChaosOutfitAsset/Source/ChaosOutfitAssetEngine/Public/ChaosOutfitAsset/OutfitCollection.h
@@ -4,6 +4,8 @@
 
 #include "UObject/NameTypes.h"
 
+#define UE_API CHAOSOUTFITASSETENGINE_API
+
 namespace UE::Chaos::OutfitAsset
 {
 	/**
@@ -13,43 +15,43 @@ namespace UE::Chaos::OutfitAsset
 	{
 		namespace Group
 		{
-			CHAOSOUTFITASSETENGINE_API extern const FName Outfits;
-			CHAOSOUTFITASSETENGINE_API extern const FName Pieces;
-			CHAOSOUTFITASSETENGINE_API extern const FName BodySizes;
-			CHAOSOUTFITASSETENGINE_API extern const FName BodyParts;
-			CHAOSOUTFITASSETENGINE_API extern const FName Measurements;
+			UE_API extern const FName Outfits;
+			UE_API extern const FName Pieces;
+			UE_API extern const FName BodySizes;
+			UE_API extern const FName BodyParts;
+			UE_API extern const FName Measurements;
 		}
 		namespace Attribute
 		{
 			namespace Outfits
 			{
-				CHAOSOUTFITASSETENGINE_API extern const FName Guid;  // FGuid
-				CHAOSOUTFITASSETENGINE_API extern const FName BodySize;  // int32 -> BodySizes
-				CHAOSOUTFITASSETENGINE_API extern const FName Name;  // FString
-				CHAOSOUTFITASSETENGINE_API extern const FName PiecesStart;  // int32 -> Pieces
-				CHAOSOUTFITASSETENGINE_API extern const FName PiecesCount;  // int32
+				UE_API extern const FName Guid;  // FGuid
+				UE_API extern const FName BodySize;  // int32 -> BodySizes
+				UE_API extern const FName Name;  // FString
+				UE_API extern const FName PiecesStart;  // int32 -> Pieces
+				UE_API extern const FName PiecesCount;  // int32
 			}
 			namespace Pieces
 			{
-				CHAOSOUTFITASSETENGINE_API extern const FName Guid;  // FGuid
-				CHAOSOUTFITASSETENGINE_API extern const FName Name;  // FString
+				UE_API extern const FName Guid;  // FGuid
+				UE_API extern const FName Name;  // FString
 			}
 			namespace BodySizes
 			{
-				CHAOSOUTFITASSETENGINE_API extern const FName Name;  // FString
-				CHAOSOUTFITASSETENGINE_API extern const FName BodyPartsStart;  // int32 -> BodyParts
-				CHAOSOUTFITASSETENGINE_API extern const FName BodyPartsCount;  // int32
+				UE_API extern const FName Name;  // FString
+				UE_API extern const FName BodyPartsStart;  // int32 -> BodyParts
+				UE_API extern const FName BodyPartsCount;  // int32
 			}
 			namespace BodyParts
 			{
-				CHAOSOUTFITASSETENGINE_API extern const FName SkeletalMesh;  // FString
-				CHAOSOUTFITASSETENGINE_API extern const FName RBFInterpolationSampleIndices;  // TArray<int32>
-				CHAOSOUTFITASSETENGINE_API extern const FName RBFInterpolationSampleRestPositions;  // TArray<FVector3f>
-				CHAOSOUTFITASSETENGINE_API extern const FName RBFInterpolationWeights;  // TArray<float>
+				UE_API extern const FName SkeletalMesh;  // FString
+				UE_API extern const FName RBFInterpolationSampleIndices;  // TArray<int32>
+				UE_API extern const FName RBFInterpolationSampleRestPositions;  // TArray<FVector3f>
+				UE_API extern const FName RBFInterpolationWeights;  // TArray<float>
 			}
 			namespace Measurements
 			{
-				CHAOSOUTFITASSETENGINE_API extern const FName Name;  // FString
+				UE_API extern const FName Name;  // FString
 			}
 		}
 	}  // namespace OutfitCollection
@@ -60,3 +62,5 @@ namespace UE::Chaos::OutfitAsset
 	/** Default number of RBF Interpolation points used for RBF Resizing */
 	inline constexpr int32 DefaultNumRBFInterpolationPoints = 1500;
 }  // namespace UE::Chaos::OutfitAsset
+
+#undef UE_API
