Chaos Outfit Asset - Added the Outfit object to the OutfitAsset so that it can rebuild the render data on loading without having to re-evaluate the Dataflow graph.

#rnx

diff --git a/Engine/Plugins/ChaosClothAsset/Source/ChaosClothAssetEngine/Private/ChaosClothAsset/ClothAsset.cpp b/Engine/Plugins/ChaosClothAsset/Source/ChaosClothAssetEngine/Private/ChaosClothAsset/ClothAsset.cpp
--- a/Engine/Plugins/ChaosClothAsset/Source/ChaosClothAssetEngine/Private/ChaosClothAsset/ClothAsset.cpp
+++ b/Engine/Plugins/ChaosClothAsset/Source/ChaosClothAssetEngine/Private/ChaosClothAsset/ClothAsset.cpp
@@ -300,75 +300,7 @@ void UChaosClothAsset::PostEditChangeProperty(FPropertyChangedEvent& PropertyCha
 	}
 	Super::PostEditChangeProperty(PropertyChangedEvent);
 }
-#endif // #if WITH_EDITOR
-
-void UChaosClothAsset::GetResourceSizeEx(FResourceSizeEx& CumulativeResourceSize)
-{
-	Super::GetResourceSizeEx(CumulativeResourceSize);
-
-	if (GetResourceForRendering())
-	{
-		GetResourceForRendering()->GetResourceSizeEx(CumulativeResourceSize);
-	}
-
-	if (ClothSimulationModel)
-	{
-		ClothSimulationModel->GetResourceSizeEx(CumulativeResourceSize);
-	}
-
-#if !UE_BUILD_SHIPPING
-	FString MemoryReport;
-	MemoryReport.Appendf(TEXT("---- Memory report for Cloth Asset [%s] ----"), *this->GetName());
-
-	FResourceSizeEx RenderDataResourceSize;
-	if (GetResourceForRendering())
-	{
-		for (int32 LodIndex = 0; LodIndex < GetResourceForRendering()->LODRenderData.Num(); ++LodIndex)
-		{
-			FResourceSizeEx LODRenderDataResourceSize;
-			GetResourceForRendering()->LODRenderData[LodIndex].GetResourceSizeEx(LODRenderDataResourceSize);
-			MemoryReport.Appendf(TEXT("\n LODRenderData LOD%d size: %ld bytes"), LodIndex, LODRenderDataResourceSize.GetTotalMemoryBytes());
-		}
-
-		GetResourceForRendering()->GetResourceSizeEx(RenderDataResourceSize);
-	}
-	MemoryReport.Appendf(TEXT("\n Total RenderData size: %ld bytes"), RenderDataResourceSize.GetTotalMemoryBytes());
 
-	FResourceSizeEx ClothSimulationModelResourceSize;
-	if (ClothSimulationModel)
-	{
-		for (int32 LodIndex = 0; LodIndex < ClothSimulationModel->GetNumLods(); ++LodIndex)
-		{
-			FResourceSizeEx ClothSimulationLodModelResourceSize;
-			ClothSimulationModel->ClothSimulationLodModels[LodIndex].GetResourceSizeEx(ClothSimulationLodModelResourceSize);
-
-			const FName Tag(*FString::Printf(TEXT("ClothSimulationLodModel[%d]"), LodIndex));
-			MemoryReport.Appendf(TEXT("\n ClothSimulationLodModel LOD%d size: %ld bytes"), LodIndex, ClothSimulationLodModelResourceSize.GetTotalMemoryBytes());
-		}
-
-		ClothSimulationModel->GetResourceSizeEx(ClothSimulationModelResourceSize);
-	}
-	MemoryReport.Appendf(TEXT("\n Total ClothSimulationModel size: %ld bytes"), ClothSimulationModelResourceSize.GetTotalMemoryBytes());
-
-	const int64 TotalResourceSize = RenderDataResourceSize.GetTotalMemoryBytes() + ClothSimulationModelResourceSize.GetTotalMemoryBytes();
-	MemoryReport.Appendf(
-		TEXT("\n Total resource size for Cloth Asset [%s]: %ld bytes (%.3f MB)"),
-		*this->GetName(),
-		TotalResourceSize,
-		(float)TotalResourceSize / (1024.f * 1024.f));
-
-	const int64 TotalSize = CumulativeResourceSize.GetTotalMemoryBytes();
-	MemoryReport.Appendf(
-		TEXT("\n Total size for Cloth Asset [%s]: %ld bytes (%.3f MB)"),
-		*this->GetName(),
-		TotalSize,
-		(float)TotalSize / (1024.f * 1024.f));
-
-	UE_LOG(LogChaosClothAsset, Display, TEXT("\n%s"), *MemoryReport);
-#endif
-}
-
-#if WITH_EDITOR
 void UChaosClothAsset::ExecuteBuildInternal(FSkinnedAssetBuildContext& Context)
 {
 	TRACE_CPUPROFILER_EVENT_SCOPE(UChaosClothAsset::ExecuteBuildInternal);
diff --git a/Engine/Plugins/ChaosClothAsset/Source/ChaosClothAssetEngine/Private/ChaosClothAsset/ClothAssetBase.cpp b/Engine/Plugins/ChaosClothAsset/Source/ChaosClothAssetEngine/Private/ChaosClothAsset/ClothAssetBase.cpp
--- a/Engine/Plugins/ChaosClothAsset/Source/ChaosClothAssetEngine/Private/ChaosClothAsset/ClothAssetBase.cpp
+++ b/Engine/Plugins/ChaosClothAsset/Source/ChaosClothAssetEngine/Private/ChaosClothAsset/ClothAssetBase.cpp
@@ -6,6 +6,7 @@
 #endif
 #include "ChaosClothAsset/ClothAssetPrivate.h"
 #include "ChaosClothAsset/ClothComponent.h"
+#include "ChaosClothAsset/ClothSimulationModel.h"
 #include "ChaosClothAsset/SkeletalMeshConverterClassProvider.h"
 #include "Engine/Engine.h"
 #include "Engine/SkeletalMesh.h"
@@ -17,7 +18,6 @@
 #include "Rendering/SkeletalMeshRenderData.h"
 #include "UObject/FortniteMainBranchObjectVersion.h"
 
-
 #include UE_INLINE_GENERATED_CPP_BY_NAME(ClothAssetBase)
 namespace UE::Chaos::ClothAsset::Private
 {
@@ -221,6 +221,75 @@ PRAGMA_ENABLE_DEPRECATION_WARNINGS
 }
 #endif // #if WITH_EDITOR
 
+void UChaosClothAssetBase::GetResourceSizeEx(FResourceSizeEx& CumulativeResourceSize)
+{
+	Super::GetResourceSizeEx(CumulativeResourceSize);
+
+	if (GetResourceForRendering())
+	{
+		GetResourceForRendering()->GetResourceSizeEx(CumulativeResourceSize);
+	}
+
+	for (int32 ModelIndex = 0; ModelIndex < GetNumClothSimulationModels(); ++ModelIndex)
+	{
+		if (const TSharedPtr<FChaosClothSimulationModel> ClothSimulationModel = ConstCastSharedPtr<FChaosClothSimulationModel>(GetClothSimulationModel(ModelIndex)))
+		{
+			ClothSimulationModel->GetResourceSizeEx(CumulativeResourceSize);
+		}
+	}
+
+#if !UE_BUILD_SHIPPING
+	FString MemoryReport;
+	MemoryReport.Appendf(TEXT("---- Memory report for [%s] [%s] ----"), *StaticClass()->GetName(), *this->GetName());
+
+	FResourceSizeEx RenderDataResourceSize;
+	if (GetResourceForRendering())
+	{
+		for (int32 LodIndex = 0; LodIndex < GetResourceForRendering()->LODRenderData.Num(); ++LodIndex)
+		{
+			FResourceSizeEx LODRenderDataResourceSize;
+			GetResourceForRendering()->LODRenderData[LodIndex].GetResourceSizeEx(LODRenderDataResourceSize);
+			MemoryReport.Appendf(TEXT("\n LODRenderData LOD%d size: %ld bytes"), LodIndex, LODRenderDataResourceSize.GetTotalMemoryBytes());
+		}
+
+		GetResourceForRendering()->GetResourceSizeEx(RenderDataResourceSize);
+	}
+	MemoryReport.Appendf(TEXT("\n Total RenderData size: %ld bytes"), RenderDataResourceSize.GetTotalMemoryBytes());
+
+	FResourceSizeEx ClothSimulationModelsResourceSize;
+	for (int32 ModelIndex = 0; ModelIndex < GetNumClothSimulationModels(); ++ModelIndex)
+	{
+		if (const TSharedPtr<FChaosClothSimulationModel> ClothSimulationModel = ConstCastSharedPtr<FChaosClothSimulationModel>(GetClothSimulationModel(ModelIndex)))
+		{
+			for (int32 LodIndex = 0; LodIndex < ClothSimulationModel->GetNumLods(); ++LodIndex)
+			{
+				FResourceSizeEx ClothSimulationLodModelResourceSize;
+				ClothSimulationModel->ClothSimulationLodModels[LodIndex].GetResourceSizeEx(ClothSimulationLodModelResourceSize);
+				MemoryReport.Appendf(TEXT("\n ClothSimulationModel%d LOD%d size: %ld bytes"), ModelIndex, LodIndex, ClothSimulationLodModelResourceSize.GetTotalMemoryBytes());
+			}
+
+			ClothSimulationModel->GetResourceSizeEx(ClothSimulationModelsResourceSize);
+		}
+	}
+	MemoryReport.Appendf(TEXT("\n Total ClothSimulationModel(s) size: %ld bytes"), ClothSimulationModelsResourceSize.GetTotalMemoryBytes());
+
+	const int64 TotalResourceSize = RenderDataResourceSize.GetTotalMemoryBytes() + ClothSimulationModelsResourceSize.GetTotalMemoryBytes();
+	MemoryReport.Appendf(
+		TEXT("\n Total resource size for Cloth Asset [%s]: %ld bytes (%.3f MB)"),
+		*this->GetName(),
+		TotalResourceSize,
+		(float)TotalResourceSize / (1024.f * 1024.f));
+
+	const int64 TotalSize = CumulativeResourceSize.GetTotalMemoryBytes();
+	MemoryReport.Appendf(
+		TEXT("\n Total size for Cloth Asset [%s]: %ld bytes (%.3f MB)"),
+		*this->GetName(),
+		TotalSize,
+		(float)TotalSize / (1024.f * 1024.f));
+
+	UE_LOG(LogChaosClothAsset, Display, TEXT("\n%s"), *MemoryReport);
+#endif
+}
 
 FSkeletalMeshLODInfo* UChaosClothAssetBase::GetLODInfo(int32 Index)
 {
diff --git a/Engine/Plugins/ChaosClothAsset/Source/ChaosClothAssetEngine/Public/ChaosClothAsset/ClothAsset.h b/Engine/Plugins/ChaosClothAsset/Source/ChaosClothAssetEngine/Public/ChaosClothAsset/ClothAsset.h
--- a/Engine/Plugins/ChaosClothAsset/Source/ChaosClothAssetEngine/Public/ChaosClothAsset/ClothAsset.h
+++ b/Engine/Plugins/ChaosClothAsset/Source/ChaosClothAssetEngine/Public/ChaosClothAsset/ClothAsset.h
@@ -55,7 +55,6 @@ public:
 #if WITH_EDITOR
 	UE_API virtual void PostEditChangeProperty(FPropertyChangedEvent& PropertyChangedEvent) override;
 #endif
-	UE_API virtual void GetResourceSizeEx(FResourceSizeEx& CumulativeResourceSize) override;
 	//~ End UObject interface
 
 	//~ Begin USkinnedAsset interface
diff --git a/Engine/Plugins/ChaosClothAsset/Source/ChaosClothAssetEngine/Public/ChaosClothAsset/ClothAssetBase.h b/Engine/Plugins/ChaosClothAsset/Source/ChaosClothAssetEngine/Public/ChaosClothAsset/ClothAssetBase.h
--- a/Engine/Plugins/ChaosClothAsset/Source/ChaosClothAssetEngine/Public/ChaosClothAsset/ClothAssetBase.h
+++ b/Engine/Plugins/ChaosClothAsset/Source/ChaosClothAssetEngine/Public/ChaosClothAsset/ClothAssetBase.h
@@ -86,6 +86,7 @@ public:
 #if WITH_EDITOR
 	UE_API virtual void PostEditChangeProperty(FPropertyChangedEvent& PropertyChangedEvent) override;
 #endif
+	UE_API virtual void GetResourceSizeEx(FResourceSizeEx& CumulativeResourceSize) override;
 	//~ End UObject interface
 
 	//~ Begin USkinnedAsset interface
diff --git a/Engine/Plugins/Experimental/ChaosOutfitAsset/Source/ChaosOutfitAssetEngine/Private/ChaosOutfitAsset/OutfitAsset.cpp b/Engine/Plugins/Experimental/ChaosOutfitAsset/Source/ChaosOutfitAssetEngine/Private/ChaosOutfitAsset/OutfitAsset.cpp
--- a/Engine/Plugins/Experimental/ChaosOutfitAsset/Source/ChaosOutfitAssetEngine/Private/ChaosOutfitAsset/OutfitAsset.cpp
+++ b/Engine/Plugins/Experimental/ChaosOutfitAsset/Source/ChaosOutfitAssetEngine/Private/ChaosOutfitAsset/OutfitAsset.cpp
@@ -3,13 +3,13 @@
 #include "ChaosOutfitAsset/OutfitAsset.h"
 #include "ChaosClothAsset/ClothSimulationModel.h"
 #include "ChaosOutfitAsset/CollectionOutfitFacade.h"
+#include "ChaosOutfitAsset/OutfitAssetPrivate.h"
 #include "Engine/SkeletalMesh.h"
 #if WITH_EDITOR
 #include "DerivedDataCacheInterface.h"
 #endif
 #include "Rendering/SkeletalMeshRenderData.h"
 
-
 #include UE_INLINE_GENERATED_CPP_BY_NAME(OutfitAsset)
 
 // If Chaos outfit asset derived data needs to be rebuilt (new format, serialization differences, etc.) replace the version GUID below with a new one.
@@ -46,7 +46,7 @@ bool UChaosOutfitAsset::HasValidClothSimulationModels() const
 	return false;
 }
 
-void UChaosOutfitAsset::Build(const TObjectPtr<const UChaosOutfit> Outfit)
+void UChaosOutfitAsset::Build(const TObjectPtr<const UChaosOutfit> InOutfit)
 {
 	using namespace UE::Chaos::OutfitAsset;
 
@@ -57,14 +57,22 @@ void UChaosOutfitAsset::Build(const TObjectPtr<const UChaosOutfit> Outfit)
 	TUniquePtr<FSkeletalMeshRenderData> TempSkeletalMeshRenderData = MakeUnique<FSkeletalMeshRenderData>();
 	FReferenceSkeleton TempReferenceSkeleton;
 
-	if (Outfit)
+	if (InOutfit)
 	{
-		Outfit->CopyTo(
+		InOutfit->CopyTo(
 			Pieces,
 			TempReferenceSkeleton,
 			TempSkeletalMeshRenderData,
 			Materials,
 			OutfitCollection);
+
+#if WITH_EDITORONLY_DATA
+		if (Outfit != InOutfit)
+		{
+			const FName UniqueOutfitName = MakeUniqueObjectName(this, UChaosOutfit::StaticClass());
+			Outfit = DuplicateObject<UChaosOutfit>(InOutfit, this, UniqueOutfitName);
+		}
+#endif
 	}
 	else
 	{
@@ -74,6 +82,10 @@ void UChaosOutfitAsset::Build(const TObjectPtr<const UChaosOutfit> Outfit)
 			TempSkeletalMeshRenderData,
 			Materials,
 			OutfitCollection);
+
+#if WITH_EDITORONLY_DATA
+		Outfit = nullptr;
+#endif
 	}
 
 	// Populate the body sizes
@@ -120,6 +132,24 @@ void UChaosOutfitAsset::Serialize(FArchive& Ar)
 {
 	LLM_SCOPE_BYNAME(TEXT("Physics/Cloth"));
 	Super::Serialize(Ar);
+
+#if WITH_EDITORONLY_DATA
+	if (Outfit)
+#endif
+	{
+		bool bCooked = Ar.IsCooking();
+		Ar << bCooked;
+
+		if (bCooked && !IsTemplate() && !Ar.IsCountingMemory())  // Counting of these resources are done in GetResourceSizeEx, so skip these when counting memory
+		{
+			LLM_SCOPE_BYNAME(TEXT("Physics/ClothRendering"));
+			if (Ar.IsLoading())
+			{
+				SetResourceForRendering(MakeUnique<FSkeletalMeshRenderData>());
+			}
+			GetResourceForRendering()->Serialize(Ar, this);
+		}
+	}
 }
 
 void UChaosOutfitAsset::PostLoad()
@@ -127,8 +157,23 @@ void UChaosOutfitAsset::PostLoad()
 	LLM_SCOPE_BYNAME(TEXT("Physics/Cloth"));
 	Super::PostLoad();
 
-	// Re-evaluate the Dataflow
-	GetDataflowInstance().UpdateOwnerAsset();  // TODO: Serialize the render data to DDC instead
+#if WITH_EDITORONLY_DATA
+	if (Outfit)
+	{
+		// Rebuild the outfit
+		Build(Outfit);
+	}
+	else
+	{
+		// Re-evaluate the Dataflow (legacy PostLoad behavior from before the Outfit object was saved)
+		GetDataflowInstance().UpdateOwnerAsset();
+		if (UPackage* const Package = GetOutermost())
+		{
+			Package->SetDirtyFlag(true);
+		}
+		UE_LOG(LogChaosOutfitAsset, Warning, TEXT("Outfit Asset [%s] needs to be re-saved."), *GetName());
+	}
+#endif
 
 	if (FApp::CanEverRender())
 	{
diff --git a/Engine/Plugins/Experimental/ChaosOutfitAsset/Source/ChaosOutfitAssetEngine/Private/ChaosOutfitAsset/OutfitCollection.cpp b/Engine/Plugins/Experimental/ChaosOutfitAsset/Source/ChaosOutfitAssetEngine/Private/ChaosOutfitAsset/OutfitCollection.cpp
--- a/Engine/Plugins/Experimental/ChaosOutfitAsset/Source/ChaosOutfitAssetEngine/Private/ChaosOutfitAsset/OutfitCollection.cpp
+++ b/Engine/Plugins/Experimental/ChaosOutfitAsset/Source/ChaosOutfitAssetEngine/Private/ChaosOutfitAsset/OutfitCollection.cpp
@@ -22,7 +22,7 @@ namespace UE::Chaos::OutfitAsset
 			// Outfits group
 			namespace Outfits
 			{
-				const FName Guid(TEXT("Guid"));  // FGuid
+				const FName Guid(TEXT("OutfitGuid"));  // FGuid, cannot simply be called Guid or it would get stripped from the ManagedArrayCollection serialization
 				const FName BodySize(TEXT("BodySize"));  // int32 -> BodySizes
 				const FName Name(TEXT("Name"));  // FString
 				const FName PiecesStart(TEXT("PiecesStart"));  // int32 -> OutfitPieces
@@ -31,7 +31,7 @@ namespace UE::Chaos::OutfitAsset
 			// Pieces group
 			namespace Pieces
 			{
-				const FName Guid(TEXT("Guid"));  // FGuid
+				const FName Guid(TEXT("PieceGuid"));  // FGuid, cannot simply be called Guid or it would get stripped from the ManagedArrayCollection serialization
 				const FName Name(TEXT("Name"));  // FString
 			}
 			// BodySizes group
diff --git a/Engine/Plugins/Experimental/ChaosOutfitAsset/Source/ChaosOutfitAssetEngine/Public/ChaosOutfitAsset/OutfitAsset.h b/Engine/Plugins/Experimental/ChaosOutfitAsset/Source/ChaosOutfitAssetEngine/Public/ChaosOutfitAsset/OutfitAsset.h
--- a/Engine/Plugins/Experimental/ChaosOutfitAsset/Source/ChaosOutfitAssetEngine/Public/ChaosOutfitAsset/OutfitAsset.h
+++ b/Engine/Plugins/Experimental/ChaosOutfitAsset/Source/ChaosOutfitAssetEngine/Public/ChaosOutfitAsset/OutfitAsset.h
@@ -6,26 +6,28 @@
 #include "ChaosOutfitAsset/Outfit.h"
 #include "OutfitAsset.generated.h"
 
+#define UE_API CHAOSOUTFITASSETENGINE_API
+
 struct FChaosOutfitPiece;
 struct FSkeletalMaterial;
 
 /**
  * Outfit asset for character clothing and simulation.
  */
-UCLASS(Experimental, HideCategories = Object, BlueprintType, PrioritizeCategories = ("Dataflow"))
-class CHAOSOUTFITASSETENGINE_API UChaosOutfitAsset final : public UChaosClothAssetBase
+UCLASS(Experimental, MinimalAPI, HideCategories = Object, BlueprintType, PrioritizeCategories = ("Dataflow"))
+class UChaosOutfitAsset final : public UChaosClothAssetBase
 {
 	GENERATED_BODY()
 public:
-	UChaosOutfitAsset(const FObjectInitializer& ObjectInitializer);
-	UChaosOutfitAsset(FVTableHelper& Helper);  // This is declared so we can use TUniquePtr<FClothSimulationModel> with just a forward declare of that class
-	virtual ~UChaosOutfitAsset() override;
+	UE_API UChaosOutfitAsset(const FObjectInitializer& ObjectInitializer);
+	UE_API UChaosOutfitAsset(FVTableHelper& Helper);  // This is declared so we can use TUniquePtr<FClothSimulationModel> with just a forward declare of that class
+	UE_API virtual ~UChaosOutfitAsset() override;
 
-	void Build(const TObjectPtr<const UChaosOutfit> Outfit);
+	UE_API void Build(const TObjectPtr<const UChaosOutfit> InOutfit);
 
 	//~ Begin UObject interface
-	virtual void Serialize(FArchive& Ar) override;
-	virtual void PostLoad() override;
+	UE_API virtual void Serialize(FArchive& Ar) override;
+	UE_API virtual void PostLoad() override;
 	//~ End UObject interface
 
 	const FManagedArrayCollection& GetOutfitCollection() const
@@ -35,16 +37,16 @@ public:
 
 private:
 	//~ Begin UChaosClothAssetBase interface
-	virtual bool HasValidClothSimulationModels() const override;
+	UE_API virtual bool HasValidClothSimulationModels() const override;
 	virtual int32 GetNumClothSimulationModels() const override
 	{
 		return Pieces.Num();
 	}
-	virtual FName GetClothSimulationModelName(int32 ModelIndex) const override;
-	virtual TSharedPtr<const FChaosClothSimulationModel> GetClothSimulationModel(int32 ModelIndex) const override;
-	virtual const TArray<TSharedRef<const FManagedArrayCollection>>& GetCollections(int32 ModelIndex) const override;
-	virtual const UPhysicsAsset* GetPhysicsAssetForModel(int32 ModelIndex) const override;
-	virtual FGuid GetAssetGuid(int32 ModelIndex) const override;
+	UE_API virtual FName GetClothSimulationModelName(int32 ModelIndex) const override;
+	UE_API virtual TSharedPtr<const FChaosClothSimulationModel> GetClothSimulationModel(int32 ModelIndex) const override;
+	UE_API virtual const TArray<TSharedRef<const FManagedArrayCollection>>& GetCollections(int32 ModelIndex) const override;
+	UE_API virtual const UPhysicsAsset* GetPhysicsAssetForModel(int32 ModelIndex) const override;
+	UE_API virtual FGuid GetAssetGuid(int32 ModelIndex) const override;
 	//~ End UChaosClothAssetBase interface
 
 	//~ Begin USkinnedAsset interface
@@ -62,7 +64,7 @@ private:
 	}
 	virtual void SetSkeleton(USkeleton* InSkeleton) override {}
 #if WITH_EDITOR
-	virtual FString BuildDerivedDataKey(const ITargetPlatform* TargetPlatform) override;
+	UE_API virtual FString BuildDerivedDataKey(const ITargetPlatform* TargetPlatform) override;
 	virtual bool IsInitialBuildDone() const override
 	{
 		return true;
@@ -76,7 +78,7 @@ private:
 #endif
 	//~ End USkinnedAsset interface
 
-	void CalculateBounds();
+	UE_API void CalculateBounds();
 
 	UPROPERTY()
 	TArray<FChaosOutfitPiece> Pieces;
@@ -86,4 +88,11 @@ private:
 
 	UPROPERTY()
 	FManagedArrayCollection OutfitCollection;
+
+#if WITH_EDITORONLY_DATA
+	UPROPERTY(Instanced)
+	TObjectPtr<UChaosOutfit> Outfit;  // Outfit source model used for generating this outfit asset
+#endif
 };
+
+#undef UE_API
diff --git a/Engine/Source/Runtime/Experimental/Chaos/Private/GeometryCollection/ManagedArrayCollection.cpp b/Engine/Source/Runtime/Experimental/Chaos/Private/GeometryCollection/ManagedArrayCollection.cpp
--- a/Engine/Source/Runtime/Experimental/Chaos/Private/GeometryCollection/ManagedArrayCollection.cpp
+++ b/Engine/Source/Runtime/Experimental/Chaos/Private/GeometryCollection/ManagedArrayCollection.cpp
@@ -4,6 +4,7 @@
 #include "GeometryCollection/GeometryCollectionAlgo.h"
 #include "Templates/UnrealTemplate.h"
 #include "Chaos/ChaosArchive.h"
+#include "Serialization/MemoryWriter.h"
 #include "UObject/UE5MainStreamObjectVersion.h"
 
 #include UE_INLINE_GENERATED_CPP_BY_NAME(ManagedArrayCollection)
@@ -26,6 +27,40 @@ FManagedArrayCollection& FManagedArrayCollection::operator=(const FManagedArrayC
 FManagedArrayCollection::FManagedArrayCollection(FManagedArrayCollection&&) = default;
 FManagedArrayCollection& FManagedArrayCollection::operator=(FManagedArrayCollection&&)= default;
 
+bool FManagedArrayCollection::operator==(const FManagedArrayCollection& Other) const
+{
+	if (&Other == this || (Other.Map.IsEmpty() && Map.IsEmpty()))
+	{
+		return true;
+	}
+	if (Other.Map.Num() != Map.Num())
+	{
+		return false;
+	}
+	for (const TTuple<FKeyType, FValueType>& Entry : Map)
+	{
+		const FName& AttributeName = Entry.Key.Get<0>();
+		const FName& GroupName = Entry.Key.Get<1>();
+		if (!Other.HasAttribute(AttributeName, GroupName))
+		{
+			return false;
+		}
+		if (Other.NumElements(GroupName) != NumElements(GroupName))
+		{
+			return false;
+		}
+	}
+	auto GetSerializedData = [](const FManagedArrayCollection& ManagedArrayCollection) -> TArray<uint8>
+		{
+			TArray<uint8> SerializedData;
+			FMemoryWriter Ar(SerializedData);
+			const_cast<FManagedArrayCollection&>(ManagedArrayCollection).Serialize(Ar);
+			Ar.Close();
+			return SerializedData;
+		};
+	return GetSerializedData(*this) == GetSerializedData(Other);
+}
+
 void FManagedArrayCollection::AddGroup(FName Group)
 {
 	ensure(!GroupInfo.Contains(Group));
diff --git a/Engine/Source/Runtime/Experimental/Chaos/Public/GeometryCollection/ManagedArrayCollection.h b/Engine/Source/Runtime/Experimental/Chaos/Public/GeometryCollection/ManagedArrayCollection.h
--- a/Engine/Source/Runtime/Experimental/Chaos/Public/GeometryCollection/ManagedArrayCollection.h
+++ b/Engine/Source/Runtime/Experimental/Chaos/Public/GeometryCollection/ManagedArrayCollection.h
@@ -68,6 +68,8 @@ public:
 	CHAOS_API FManagedArrayCollection(FManagedArrayCollection&&);
 	CHAOS_API FManagedArrayCollection& operator=(FManagedArrayCollection&&);
 
+	CHAOS_API bool operator==(const FManagedArrayCollection& Other) const;  // Slow, required for comparing to a default collection during UPROPERTY serialization
+
 	static CHAOS_API int8 Invalid;
 	typedef EManagedArrayType EArrayType;
 	/**
@@ -854,6 +856,7 @@ template<> struct TStructOpsTypeTraits<FManagedArrayCollection> : public TStruct
 	enum
 	{
 		WithSerializer = true,
+		WithIdenticalViaEquality = true  // Required for usage as a UPROPERTY, otherwise the reflection system won't see a difference with a default collection when serializing
 	};
 };
 
