Fix Outfit render data to serialize correctly when there is no cloth (simulation) data.

diff --git a/Engine/Plugins/Experimental/ChaosOutfitAsset/Source/ChaosOutfitAssetEngine/Private/ChaosOutfitAsset/Outfit.cpp b/Engine/Plugins/Experimental/ChaosOutfitAsset/Source/ChaosOutfitAssetEngine/Private/ChaosOutfitAsset/Outfit.cpp
--- a/Engine/Plugins/Experimental/ChaosOutfitAsset/Source/ChaosOutfitAssetEngine/Private/ChaosOutfitAsset/Outfit.cpp
+++ b/Engine/Plugins/Experimental/ChaosOutfitAsset/Source/ChaosOutfitAssetEngine/Private/ChaosOutfitAsset/Outfit.cpp
@@ -56,6 +56,18 @@ struct UChaosOutfit::FLODRenderData
 		return StaticVertexBuffers.PositionVertexBuffer.GetNumVertices();
 	}
 
+	bool HasClothData() const
+	{
+		for (int32 SectionIdx = 0; SectionIdx < RenderSections.Num(); SectionIdx++)
+		{
+			if (RenderSections[SectionIdx].HasClothingData())
+			{
+				return true;
+			}
+		}
+		return false;	
+	}
+
 	friend FArchive& operator<<(FArchive& Ar, UChaosOutfit::FLODRenderData& LODRenderData)
 	{
 		static const FName MultiSizeIndexContainerName(TEXT("UChaosOutfit::FLODRenderData.MultiSizeIndexContainer"));
@@ -78,9 +90,11 @@ struct UChaosOutfit::FLODRenderData
 		LODRenderData.StaticVertexBuffers.StaticMeshVertexBuffer.SetResourceName(StaticMeshVertexBufferName);
 
 		Ar << LODRenderData.SkinWeightVertexBuffer;
-
-		Ar << LODRenderData.ClothVertexBuffer;
-		LODRenderData.ClothVertexBuffer.SetResourceName(ClothVertexBufferName);
+		if (LODRenderData.HasClothData()) // This requires that RenderSections have already been serialized/deserialized
+		{
+			Ar << LODRenderData.ClothVertexBuffer;
+			LODRenderData.ClothVertexBuffer.SetResourceName(ClothVertexBufferName);
+		}
 
 		Ar << LODRenderData.ActiveBoneIndices;
 		Ar << LODRenderData.RequiredBones;
