Chaos Outfit Asset:
- Added template selection dialog when creating a new Outfit Asset.
- Fixed crash when opening an Outfit Asset without Dataflow.
- Updated the plugin icons.

#rnx

#virtualized

diff --git a/Engine/Plugins/ChaosClothAsset/Source/ChaosClothAssetEngine/Public/ChaosClothAsset/ClothAssetBase.h b/Engine/Plugins/ChaosClothAsset/Source/ChaosClothAssetEngine/Public/ChaosClothAsset/ClothAssetBase.h
--- a/Engine/Plugins/ChaosClothAsset/Source/ChaosClothAssetEngine/Public/ChaosClothAsset/ClothAssetBase.h
+++ b/Engine/Plugins/ChaosClothAsset/Source/ChaosClothAssetEngine/Public/ChaosClothAsset/ClothAssetBase.h
@@ -17,8 +17,9 @@ struct FManagedArrayCollection;
 
 /**
  * Base cloth asset class.
+ * Any object slot based on this type can be either a Cloth Asset or an Outfit Asset.
  */
-UCLASS(MinimalAPI, ABSTRACT)
+UCLASS(MinimalAPI, Abstract, BlueprintType)
 class UChaosClothAssetBase
 	: public USkinnedAsset
 	, public IDataflowContentOwner
diff --git a/Engine/Plugins/Experimental/ChaosOutfitAsset/Source/ChaosOutfitAssetEditor/ChaosOutfitAssetEditor.Build.cs b/Engine/Plugins/Experimental/ChaosOutfitAsset/Source/ChaosOutfitAssetEditor/ChaosOutfitAssetEditor.Build.cs
--- a/Engine/Plugins/Experimental/ChaosOutfitAsset/Source/ChaosOutfitAssetEditor/ChaosOutfitAssetEditor.Build.cs
+++ b/Engine/Plugins/Experimental/ChaosOutfitAsset/Source/ChaosOutfitAssetEditor/ChaosOutfitAssetEditor.Build.cs
@@ -28,6 +28,7 @@ public class ChaosOutfitAssetEditor : ModuleRules
 				"Slate",
 				"SlateCore",
 				"ToolMenus",
+				"ToolWidgets",
 				"UnrealEd"
 			}
 		);
diff --git a/Engine/Plugins/Experimental/ChaosOutfitAsset/Source/ChaosOutfitAssetEditor/Private/ChaosOutfitAsset/AssetDefinition_OutfitAsset.cpp b/Engine/Plugins/Experimental/ChaosOutfitAsset/Source/ChaosOutfitAssetEditor/Private/ChaosOutfitAsset/AssetDefinition_OutfitAsset.cpp
--- a/Engine/Plugins/Experimental/ChaosOutfitAsset/Source/ChaosOutfitAssetEditor/Private/ChaosOutfitAsset/AssetDefinition_OutfitAsset.cpp
+++ b/Engine/Plugins/Experimental/ChaosOutfitAsset/Source/ChaosOutfitAssetEditor/Private/ChaosOutfitAsset/AssetDefinition_OutfitAsset.cpp
@@ -4,6 +4,7 @@
 #include "ChaosOutfitAsset/OutfitAsset.h"
 #include "Dataflow/DataflowEditor.h"
 #include "ThumbnailRendering/SceneThumbnailInfo.h"
+#include "Toolkits/SimpleAssetEditor.h"
 
 #define LOCTEXT_NAMESPACE "AssetDefinition_OutfitAsset"
 
@@ -43,16 +44,23 @@ EAssetCommandResult UAssetDefinition_OutfitAsset::OpenAssets(const FAssetOpenArg
 	TArray<UChaosOutfitAsset*> OutfitObjects = OpenArgs.LoadObjects<UChaosOutfitAsset>();
 
 	// For now the Dataflow editor only works on one asset at a time
-	if (OutfitObjects.Num() > 0)
+	if (OutfitObjects.Num() > 0 && OutfitObjects[0])
 	{
-		UAssetEditorSubsystem* const AssetEditorSubsystem = GEditor->GetEditorSubsystem<UAssetEditorSubsystem>();
-		UDataflowEditor* const AssetEditor = NewObject<UDataflowEditor>(AssetEditorSubsystem, NAME_None, RF_Transient);
-		AssetEditor->RegisterToolCategories({"General", "Cloth"});
-
-		static const TCHAR* const Path = TEXT("/ChaosOutfitAsset/BP_OutfitPreview.BP_OutfitPreview_C");
-		const TSubclassOf<AActor> PreviewClass = StaticLoadClass(AActor::StaticClass(), nullptr, Path);
-		AssetEditor->Initialize({ OutfitObjects[0] }, PreviewClass);
+		if (OutfitObjects[0]->GetDataflowInstance().GetDataflowAsset())
+		{
+			// Has a Dataflow Asset, open in the Dataflow editor
+			UAssetEditorSubsystem* const AssetEditorSubsystem = GEditor->GetEditorSubsystem<UAssetEditorSubsystem>();
+			UDataflowEditor* const AssetEditor = NewObject<UDataflowEditor>(AssetEditorSubsystem, NAME_None, RF_Transient);
+			AssetEditor->RegisterToolCategories({"General", "Cloth"});
 
+			static const TCHAR* const Path = TEXT("/ChaosOutfitAsset/BP_OutfitPreview.BP_OutfitPreview_C");
+			const TSubclassOf<AActor> PreviewClass = StaticLoadClass(AActor::StaticClass(), nullptr, Path);
+			AssetEditor->Initialize({ OutfitObjects[0] }, PreviewClass);
+		}
+		else
+		{
+			FSimpleAssetEditor::CreateEditor(EToolkitMode::Standalone, OpenArgs.ToolkitHost, OutfitObjects[0]);
+		}
 		return EAssetCommandResult::Handled;
 	}
 	return EAssetCommandResult::Unhandled;
diff --git a/Engine/Plugins/Experimental/ChaosOutfitAsset/Source/ChaosOutfitAssetEditor/Private/ChaosOutfitAsset/OutfitAssetFactory.cpp b/Engine/Plugins/Experimental/ChaosOutfitAsset/Source/ChaosOutfitAssetEditor/Private/ChaosOutfitAsset/OutfitAssetFactory.cpp
--- a/Engine/Plugins/Experimental/ChaosOutfitAsset/Source/ChaosOutfitAssetEditor/Private/ChaosOutfitAsset/OutfitAssetFactory.cpp
+++ b/Engine/Plugins/Experimental/ChaosOutfitAsset/Source/ChaosOutfitAssetEditor/Private/ChaosOutfitAsset/OutfitAssetFactory.cpp
@@ -4,10 +4,13 @@
 #include "AssetRegistry/AssetRegistryModule.h"
 #include "ChaosOutfitAsset/OutfitAsset.h"
 #include "Dataflow/DataflowObject.h"
+#include "Dialog/SMessageDialog.h"
 #include "Misc/PackageName.h"
 
 #include UE_INLINE_GENERATED_CPP_BY_NAME(OutfitAssetFactory)
 
+#define LOCTEXT_NAMESPACE "OutfitAssetFactory"
+
 UChaosOutfitAssetFactory::UChaosOutfitAssetFactory(const FObjectInitializer& ObjectInitializer)
 	: Super(ObjectInitializer)
 {
@@ -34,17 +37,61 @@ UObject* UChaosOutfitAssetFactory::FactoryCreateNew(UClass* Class, UObject* Pare
 	}
 	UPackage* const DataflowPackage = CreatePackage(*DataflowPackageName);
 
-	// Load the cloth template into the new Dataflow asset
-	UDataflow* const Template = LoadObject<UDataflow>(DataflowPackage, TEXT("/ChaosOutfitAsset/OutfitAssetTemplate.OutfitAssetTemplate"));
-	UDataflow* const Dataflow = DuplicateObject(Template, DataflowPackage, FName(DataflowName));
+	// Select the template
+	TSharedRef<SMessageDialog> SelectTemplateMessageDialog = SNew(SMessageDialog)
+		.Title(FText(LOCTEXT("SelectTemplateTitle", "Select an Outfit Asset Template")))
+		.Message(LOCTEXT("SelectTemplateMessage", "Select a template for this newly created Outfit Asset:"))
+		.Buttons(
+			{
+				SMessageDialog::FButton(LOCTEXT("SelectNoTemplate", "No Dataflow"))
+					.SetToolTipText(LOCTEXT("SelectNoTemplateTooltip", "Don't add a Dataflow to this Outfit Asset. Useful when a Dataflow already exists."))
+					.SetPrimary(false),
+				SMessageDialog::FButton(LOCTEXT("SelectEmptyTemplate", "Empty Dataflow"))
+					.SetToolTipText(LOCTEXT("SelectNoTemplateTooltip", "Add an empty Dataflow with an Outfit Asset Terminal node."))
+					.SetPrimary(false),
+				SMessageDialog::FButton(LOCTEXT("SelectSimpleOutfitTemplate", "Simple Outfit"))
+					.SetToolTipText(LOCTEXT("SelectNoTemplateTooltip", "Add a Dataflow with a simple Cloth Asset aggregator graph. Allows to simulate multiple Cloth Assets from the same ChaosClothComponent."))
+					.SetPrimary(true),
+				SMessageDialog::FButton(LOCTEXT("SelectResizableOutfitTemplate", "Resizable Outfit"))
+					.SetToolTipText(LOCTEXT("SelectNoTemplateTooltip", "Add a Dataflow that builds a single resizable garment from multiple Cloth Assets. Body sizes will have to be provided in addition to the mutiple Cloth Assets."))
+					.SetPrimary(false),
+				SMessageDialog::FButton(LOCTEXT("SelectResizingGraph", "Resizing Graph"))
+					.SetToolTipText(LOCTEXT("SelectNoTemplateTooltip", "Resizing graph to test a resizable Outfit Asset."))
+					.SetPrimary(false),
+			});
 
-	Dataflow->MarkPackageDirty();
+	FString OutfitAssetTemplate;
+	switch (SelectTemplateMessageDialog->ShowModal())
+	{
+	default:
+	case 0:  // No Dataflow
+		break;
+	case 1:  // Empty Dataflow
+		OutfitAssetTemplate = TEXT("/ChaosOutfitAsset/EmptyOutfitAssetTemplate.EmptyOutfitAssetTemplate");
+		break;
+	case 2:  // Simple Outfit
+		OutfitAssetTemplate = TEXT("/ChaosOutfitAsset/OutfitAssetTemplate.OutfitAssetTemplate");
+		break;
+	case 3:  // Resizable Outfit
+		OutfitAssetTemplate = TEXT("/ChaosOutfitAsset/MakeResizableOutfitTemplate.MakeResizableOutfitTemplate");
+		break;
+	case 4:  // Resizing Graph
+		OutfitAssetTemplate = TEXT("/ChaosOutfitAsset/ResizeOutfitTemplate.ResizeOutfitTemplate");
+		break;
+	}
 
-	// Notify the asset registry
-	FAssetRegistryModule::AssetCreated(Dataflow);
+	// Load the cloth template into the new Dataflow asset
+	UDataflow* const Template = OutfitAssetTemplate.IsEmpty() ? nullptr: LoadObject<UDataflow>(DataflowPackage, *OutfitAssetTemplate);
+	if (UDataflow* const Dataflow = Template ? DuplicateObject(Template, DataflowPackage, FName(DataflowName)) : nullptr)
+	{
+		Dataflow->MarkPackageDirty();
+
+		// Notify the asset registry
+		FAssetRegistryModule::AssetCreated(Dataflow);
 
-	// Set the Dataflow to the cloth asset
-	OutfitAsset->SetDataflow(Dataflow);
+		// Set the Dataflow to the cloth asset
+		OutfitAsset->SetDataflow(Dataflow);
+	}
 
 	return OutfitAsset;
 }
@@ -53,3 +100,5 @@ FString UChaosOutfitAssetFactory::GetDefaultNewAssetName() const
 {
 	return FString(TEXT("OA_NewOutfitAsset"));
 }
+
+#undef LOCTEXT_NAMESPACE
