Make LODIndex a DataflowInput for FChaosGetOutfitClothCollectionsNode

Add new  node GetOrMakeOutfromFromAsset which is much more efficient at extracting the Outfit from an existing OutfitAsset.

Updated the ResizeOutfitTemplate with these changes.

#virtualized

diff --git a/Engine/Plugins/Experimental/ChaosOutfitAsset/Source/ChaosOutfitAssetDataflowNodes/Private/ChaosOutfitAsset/ChaosOutfitAssetDataflowNodesModule.cpp b/Engine/Plugins/Experimental/ChaosOutfitAsset/Source/ChaosOutfitAssetDataflowNodes/Private/ChaosOutfitAsset/ChaosOutfitAssetDataflowNodesModule.cpp
--- a/Engine/Plugins/Experimental/ChaosOutfitAsset/Source/ChaosOutfitAssetDataflowNodes/Private/ChaosOutfitAsset/ChaosOutfitAssetDataflowNodesModule.cpp
+++ b/Engine/Plugins/Experimental/ChaosOutfitAsset/Source/ChaosOutfitAssetDataflowNodes/Private/ChaosOutfitAsset/ChaosOutfitAssetDataflowNodesModule.cpp
@@ -3,6 +3,7 @@
 #include "ChaosOutfitAsset/ClothAssetAnyType.h"
 #include "ChaosOutfitAsset/FilterSizedOutfitNode.h"
 #include "ChaosOutfitAsset/GetClothAssetNode.h"
+#include "ChaosOutfitAsset/GetOrMakeOutfitFromAssetNode.h"
 #include "ChaosOutfitAsset/GetOutfitAssetNode.h"
 #include "ChaosOutfitAsset/GetOutfitBodyPartsNode.h"
 #include "ChaosOutfitAsset/GetOutfitClothCollectionsNode.h"
@@ -40,6 +41,7 @@ namespace UE::Chaos::OutfitAsset
 			DATAFLOW_NODE_REGISTER_CREATION_FACTORY(FChaosGetOutfitClothCollectionsNode);
 			DATAFLOW_NODE_REGISTER_CREATION_FACTORY(FChaosGetOutfitRBFInterpolationDataNode);
 			DATAFLOW_NODE_REGISTER_CREATION_FACTORY(FChaosOutfitAssetFilterSizedOutfitNode);
+			DATAFLOW_NODE_REGISTER_CREATION_FACTORY(FChaosOutfitAssetGetOrMakeOutfitFromAssetNode);
 			DATAFLOW_NODE_REGISTER_CREATION_FACTORY(FChaosOutfitAssetMakeOutfitNode);
 			DATAFLOW_NODE_REGISTER_CREATION_FACTORY(FChaosOutfitAssetMakeSizedOutfitNode);
 			DATAFLOW_NODE_REGISTER_CREATION_FACTORY(FChaosOutfitAssetMergeOutfitsNode);
diff --git a/Engine/Plugins/Experimental/ChaosOutfitAsset/Source/ChaosOutfitAssetDataflowNodes/Private/ChaosOutfitAsset/GetOrMakeOutfitFromAssetNode.cpp b/Engine/Plugins/Experimental/ChaosOutfitAsset/Source/ChaosOutfitAssetDataflowNodes/Private/ChaosOutfitAsset/GetOrMakeOutfitFromAssetNode.cpp
new file mode 100644
--- /dev/null
+++ b/Engine/Plugins/Experimental/ChaosOutfitAsset/Source/ChaosOutfitAssetDataflowNodes/Private/ChaosOutfitAsset/GetOrMakeOutfitFromAssetNode.cpp
@@ -0,0 +1,38 @@
+// Copyright Epic Games, Inc. All Rights Reserved.
+
+#include "ChaosOutfitAsset/GetOrMakeOutfitFromAssetNode.h"
+#include "ChaosOutfitAsset/OutfitAsset.h"
+#include "ChaosOutfitAsset/Outfit.h"
+
+#include UE_INLINE_GENERATED_CPP_BY_NAME(GetOrMakeOutfitFromAssetNode)
+
+#define LOCTEXT_NAMESPACE "ChaosOutfitAssetGetOrMakeOutfitFromAssetNode"
+
+FChaosOutfitAssetGetOrMakeOutfitFromAssetNode::FChaosOutfitAssetGetOrMakeOutfitFromAssetNode(const UE::Dataflow::FNodeParameters& InParam, FGuid InGuid)
+	: FDataflowNode(InParam, InGuid)
+{
+	RegisterInputConnection(&OutfitAsset);
+	RegisterOutputConnection(&Outfit);
+}
+
+void FChaosOutfitAssetGetOrMakeOutfitFromAssetNode::Evaluate(UE::Dataflow::FContext& Context, const FDataflowOutput* Out) const
+{
+	if (Out->IsA(&Outfit))
+	{
+		TObjectPtr<UChaosOutfit> OutOutfit = NewObject<UChaosOutfit>();
+		if (const UChaosOutfitAsset* const InOutfitAsset = GetValue(Context, &OutfitAsset))
+		{
+#if WITH_EDITORONLY_DATA
+			if (TObjectPtr<const UChaosOutfit> AssetOutfit = InOutfitAsset->GetOutfit())
+			{
+				SetValue(Context, AssetOutfit, &Outfit);
+				return;
+			}
+#endif
+
+			OutOutfit->Add(*InOutfitAsset);
+		}
+		SetValue<TObjectPtr<const UChaosOutfit>>(Context, OutOutfit, &Outfit);
+	}
+}
+#undef LOCTEXT_NAMESPACE
diff --git a/Engine/Plugins/Experimental/ChaosOutfitAsset/Source/ChaosOutfitAssetDataflowNodes/Private/ChaosOutfitAsset/GetOrMakeOutfitFromAssetNode.h b/Engine/Plugins/Experimental/ChaosOutfitAsset/Source/ChaosOutfitAssetDataflowNodes/Private/ChaosOutfitAsset/GetOrMakeOutfitFromAssetNode.h
new file mode 100644
--- /dev/null
+++ b/Engine/Plugins/Experimental/ChaosOutfitAsset/Source/ChaosOutfitAssetDataflowNodes/Private/ChaosOutfitAsset/GetOrMakeOutfitFromAssetNode.h
@@ -0,0 +1,38 @@
+// Copyright Epic Games, Inc. All Rights Reserved.
+
+#pragma once 
+
+#include "ChaosOutfitAsset/ClothAssetAnyType.h"
+#include "Dataflow/DataflowNode.h"
+#include "GeometryCollection/ManagedArrayCollection.h"
+#include "GetOrMakeOutfitFromAssetNode.generated.h"
+
+class UChaosOutfitAsset;
+class UChaosOutfit;
+
+/** Extract the Outfit from an Outfit Asset. 
+ *  If the Outfit does not exist (e.g., this OutfitAsset has been cooked),
+ *  recreate a new one.
+ */
+USTRUCT(Meta = (Experimental, DataflowOutfit))
+struct FChaosOutfitAssetGetOrMakeOutfitFromAssetNode final : public FDataflowNode
+{
+	GENERATED_USTRUCT_BODY()
+	DATAFLOW_NODE_DEFINE_INTERNAL(FChaosOutfitAssetGetOrMakeOutfitFromAssetNode, "GetOrMakeOutfitFromAsset", "Outfit", "Get Make Outfit")
+
+public:
+	FChaosOutfitAssetGetOrMakeOutfitFromAssetNode(const UE::Dataflow::FNodeParameters& InParam, FGuid InGuid = FGuid::NewGuid());
+
+private:
+	//~ Begin FDataflowNode interface
+	virtual void Evaluate(UE::Dataflow::FContext& Context, const FDataflowOutput* Out) const override;
+	//~ End FDataflowNode interface
+
+	/** The outfit asset input. */
+	UPROPERTY(meta = (DataflowInput))
+	TObjectPtr<const UChaosOutfitAsset> OutfitAsset;
+
+	/** The outfit output. */
+	UPROPERTY(Meta = (DataflowOutput))
+	TObjectPtr<const UChaosOutfit> Outfit;
+};
diff --git a/Engine/Plugins/Experimental/ChaosOutfitAsset/Source/ChaosOutfitAssetDataflowNodes/Private/ChaosOutfitAsset/GetOutfitClothCollectionsNode.cpp b/Engine/Plugins/Experimental/ChaosOutfitAsset/Source/ChaosOutfitAssetDataflowNodes/Private/ChaosOutfitAsset/GetOutfitClothCollectionsNode.cpp
--- a/Engine/Plugins/Experimental/ChaosOutfitAsset/Source/ChaosOutfitAssetDataflowNodes/Private/ChaosOutfitAsset/GetOutfitClothCollectionsNode.cpp
+++ b/Engine/Plugins/Experimental/ChaosOutfitAsset/Source/ChaosOutfitAssetDataflowNodes/Private/ChaosOutfitAsset/GetOutfitClothCollectionsNode.cpp
@@ -11,6 +11,7 @@ FChaosGetOutfitClothCollectionsNode::FChaosGetOutfitClothCollectionsNode(const U
 	: FDataflowNode(InParam, InGuid)
 {
 	RegisterInputConnection(&Outfit);
+	RegisterInputConnection(&LodIndex);
 	RegisterOutputConnection(&Outfit, &Outfit);
 	RegisterOutputConnection(&ClothCollections);
 	RegisterOutputConnection(&NumLods);
@@ -34,12 +35,13 @@ void FChaosGetOutfitClothCollectionsNode::Evaluate(UE::Dataflow::FContext& Conte
 			SetValue(Context, OutOutfit, &Outfit);
 		}
 
-		const int32 OutNumLods = (LodIndex == INDEX_NONE) ? OutOutfit->GetNumLods() : FMath::Min(1, OutOutfit->GetNumLods());
+		const int32 InLODIndex = GetValue(Context, &LodIndex);
+		const int32 OutNumLods = (InLODIndex == INDEX_NONE) ? OutOutfit->GetNumLods() : FMath::Min(1, OutOutfit->GetNumLods());
 		const int32 OutNumPieces = OutOutfit->GetPieces().Num();
 
 		TArray<FManagedArrayCollection> OutClothCollections;
 		OutClothCollections.Reserve(OutNumLods * OutNumPieces);
-		for (const TSharedRef<const FManagedArrayCollection>& ClothCollection : OutOutfit->GetClothCollections(LodIndex))
+		for (const TSharedRef<const FManagedArrayCollection>& ClothCollection : OutOutfit->GetClothCollections(InLODIndex))
 		{
 			OutClothCollections.Emplace(*ClothCollection);
 		}
diff --git a/Engine/Plugins/Experimental/ChaosOutfitAsset/Source/ChaosOutfitAssetDataflowNodes/Private/ChaosOutfitAsset/GetOutfitClothCollectionsNode.h b/Engine/Plugins/Experimental/ChaosOutfitAsset/Source/ChaosOutfitAssetDataflowNodes/Private/ChaosOutfitAsset/GetOutfitClothCollectionsNode.h
--- a/Engine/Plugins/Experimental/ChaosOutfitAsset/Source/ChaosOutfitAssetDataflowNodes/Private/ChaosOutfitAsset/GetOutfitClothCollectionsNode.h
+++ b/Engine/Plugins/Experimental/ChaosOutfitAsset/Source/ChaosOutfitAssetDataflowNodes/Private/ChaosOutfitAsset/GetOutfitClothCollectionsNode.h
@@ -34,7 +34,7 @@ private:
 	TArray<FManagedArrayCollection> ClothCollections;
 
 	/** The LOD to output in the cloth collections array. Set to -1 to output all LODs */
-	UPROPERTY(EditAnywhere, Category = "Get Outfit Cloth Collections", DisplayName = "LOD Index", Meta = (UIMin = -1))
+	UPROPERTY(EditAnywhere, Category = "Get Outfit Cloth Collections", DisplayName = "LOD Index", Meta = (UIMin = -1, DataflowInput))
 	int32 LodIndex = -1;
 
 	/** The number of LODs output in the cloth collections array. */
diff --git a/Engine/Plugins/Experimental/ChaosOutfitAsset/Source/ChaosOutfitAssetEngine/Public/ChaosOutfitAsset/OutfitAsset.h b/Engine/Plugins/Experimental/ChaosOutfitAsset/Source/ChaosOutfitAssetEngine/Public/ChaosOutfitAsset/OutfitAsset.h
--- a/Engine/Plugins/Experimental/ChaosOutfitAsset/Source/ChaosOutfitAssetEngine/Public/ChaosOutfitAsset/OutfitAsset.h
+++ b/Engine/Plugins/Experimental/ChaosOutfitAsset/Source/ChaosOutfitAssetEngine/Public/ChaosOutfitAsset/OutfitAsset.h
@@ -35,6 +35,13 @@ public:
 		return OutfitCollection;
 	}
 
+#if WITH_EDITORONLY_DATA
+	TObjectPtr<const UChaosOutfit> GetOutfit() const
+	{
+		return Outfit;
+	}
+#endif
+
 private:
 	//~ Begin UChaosClothAssetBase interface
 	UE_API virtual bool HasValidClothSimulationModels() const override;
