Chaos Outfit Asset - Added slowtasks on RBF generation,

#rnx

diff --git a/Engine/Plugins/Experimental/ChaosOutfitAsset/Source/ChaosOutfitAssetDataflowNodes/Private/ChaosOutfitAsset/MakeSizedOutfitNode.cpp b/Engine/Plugins/Experimental/ChaosOutfitAsset/Source/ChaosOutfitAssetDataflowNodes/Private/ChaosOutfitAsset/MakeSizedOutfitNode.cpp
--- a/Engine/Plugins/Experimental/ChaosOutfitAsset/Source/ChaosOutfitAssetDataflowNodes/Private/ChaosOutfitAsset/MakeSizedOutfitNode.cpp
+++ b/Engine/Plugins/Experimental/ChaosOutfitAsset/Source/ChaosOutfitAssetDataflowNodes/Private/ChaosOutfitAsset/MakeSizedOutfitNode.cpp
@@ -5,6 +5,7 @@
 #include "ChaosOutfitAsset/ClothAssetAnyType.h"
 #include "ChaosOutfitAsset/Outfit.h"
 #include "Engine/SkeletalMesh.h"
+#include "Misc/ScopedSlowTask.h"
 
 #include UE_INLINE_GENERATED_CPP_BY_NAME(MakeSizedOutfitNode)
 
@@ -25,6 +26,8 @@ FChaosOutfitAssetMakeSizedOutfitNode::FChaosOutfitAssetMakeSizedOutfitNode(const
 
 void FChaosOutfitAssetMakeSizedOutfitNode::Evaluate(UE::Dataflow::FContext& Context, const FDataflowOutput* Out) const
 {
+	using namespace UE::Chaos::OutfitAsset;
+
 	if (Out->IsA(&Outfit) || Out->IsA(&OutfitCollection))
 	{
 		TObjectPtr<UChaosOutfit> OutOutfit = NewObject<UChaosOutfit>();
@@ -33,24 +36,63 @@ void FChaosOutfitAssetMakeSizedOutfitNode::Evaluate(UE::Dataflow::FContext& Cont
 		// Make a new GUID for this sized outfit asset
 		const FGuid OutfitGuid = FGuid::NewGuid();
 
+		// Calculate the number of source to add
+		int32 NumOutfitSources = 0;
+		for (int32 SourceIndex = 0; SourceIndex < SizedOutfitSources.Num(); ++SourceIndex)
+		{
+			const FChaosSizedOutfitSourceOrArrayType& SizedOutfitSourceOrArray = GetValue(Context, &SizedOutfitSources[SourceIndex]);
+			for (const FChaosSizedOutfitSource& SizedOutfitSource : SizedOutfitSourceOrArray.Array)
+			{
+				if (SizedOutfitSource.SourceAsset || !SizedOutfitSource.GetBodySizeName().IsEmpty())
+				{
+					++NumOutfitSources;
+				}
+			}
+		}
+
+		FScopedSlowTask SlowTask((float)NumOutfitSources, LOCTEXT("AddOutfitSources", "Adding sized outfit sources..."));
+		SlowTask.MakeDialogDelayed(0.01f);
+
+		// Add outfit sources
 		for (int32 SourceIndex = 0; SourceIndex < SizedOutfitSources.Num(); ++SourceIndex)
 		{
 			const FChaosSizedOutfitSourceOrArrayType& SizedOutfitSourceOrArray = GetValue(Context, &SizedOutfitSources[SourceIndex]);
 
 			for (const FChaosSizedOutfitSource& SizedOutfitSource : SizedOutfitSourceOrArray.Array)
 			{
-				const FString BodySizeName = SizedOutfitSource.GetBodySizeName();
+				FString BodySizeName = SizedOutfitSource.GetBodySizeName();
+				FString SourceAssetName;
 
-				if (!SizedOutfitSource.SourceAsset && !BodySizeName.IsEmpty())
+				if (!SizedOutfitSource.SourceAsset)
 				{
-					const FString SourceAssetName = SizedOutfitSource.SourceAsset->GetName();
-					Context.Warning(FString::Printf(TEXT("Empty source asset [%s] specified for size [%s]."), *SourceAssetName, *BodySizeName), this, Out);
+					if (BodySizeName.IsEmpty())
+					{
+						continue;
+					}
+					Context.Warning(FString::Printf(TEXT("Empty source asset specified for size [%s]."), *BodySizeName), this, Out);
+					SourceAssetName = TEXT("-");
+				}
+				else
+				{
+					SourceAssetName = SizedOutfitSource.SourceAsset->GetName();
+					if (BodySizeName.IsEmpty())
+					{
+						Context.Info(FString::Printf(TEXT("The source asset [%s] has no body size, and therefore will use the default body size."), *SourceAssetName), this, Out);
+						BodySizeName = DefaultBodySize.ToString();
+					}
 				}
 				if (OutfitFacade.HasBodySize(BodySizeName))
 				{
 					Context.Warning(FString::Printf(TEXT("The body size [%s] already existed and had to be overwritten."), *BodySizeName), this, Out);
 				}
 
+				SlowTask.EnterProgressFrame(1.f, 
+					FText::Format(LOCTEXT("AddOutfitSource", "Adding source asset [{0}] for body size [{1}]..."), 
+						FText::FromString(SourceAssetName),
+						FText::FromString(BodySizeName)));
+				SlowTask.TickProgress();  // The RBF weight evaluations are really slow
+				SlowTask.ForceRefresh();  // ForceRefresh could be overkill, but the progress still doesn't show in some instances
+
 				OutOutfit->Add(SizedOutfitSource, OutfitGuid);
 			}
 		}
diff --git a/Engine/Plugins/Experimental/ChaosOutfitAsset/Source/ChaosOutfitAssetEngine/Private/ChaosOutfitAsset/Outfit.cpp b/Engine/Plugins/Experimental/ChaosOutfitAsset/Source/ChaosOutfitAssetEngine/Private/ChaosOutfitAsset/Outfit.cpp
--- a/Engine/Plugins/Experimental/ChaosOutfitAsset/Source/ChaosOutfitAssetEngine/Private/ChaosOutfitAsset/Outfit.cpp
+++ b/Engine/Plugins/Experimental/ChaosOutfitAsset/Source/ChaosOutfitAssetEngine/Private/ChaosOutfitAsset/Outfit.cpp
@@ -13,6 +13,7 @@
 #include "ChaosOutfitAsset/SizedOutfitSource.h"
 #include "Engine/SkeletalMesh.h"
 #include "MeshResizing/RBFInterpolation.h"
+#include "Misc/ScopedSlowTask.h"
 #include "Rendering/SkeletalMeshLODModel.h"
 #include "Rendering/SkeletalMeshLODRenderData.h"
 #include "Rendering/SkeletalMeshRenderData.h"
@@ -20,6 +21,8 @@
 
 #include UE_INLINE_GENERATED_CPP_BY_NAME(Outfit)
 
+#define LOCTEXT_NAMESPACE "ChaosOutfit"
+
 struct UChaosOutfit::FLODRenderData
 {
 	// Half edge buffers are always regenerated on the last merge, a dummy buffer is used until then
@@ -738,6 +741,12 @@ void UChaosOutfit::Add(const FChaosSizedOutfitSource& SizedOutfitSource, const F
 					FMeshResizingRBFInterpolationData InterpData;
 					if (const FMeshDescription* const MeshDescription = SourceBodyPart->GetMeshDescription(0))
 					{
+						constexpr float NumSteps = 1.f;
+						FScopedSlowTask SlowTask(NumSteps);
+						SlowTask.MakeDialog();  // Can't delay the dialog or it won't show up
+						SlowTask.EnterProgressFrame(NumSteps, LOCTEXT("GeneratingRBFInterpolationWeights", "Generating RBF interpolation weights (please wait, this can take several minutes)..."));
+						SlowTask.TickProgress();
+						SlowTask.ForceRefresh();
 						UE::MeshResizing::FRBFInterpolation::GenerateWeights(*MeshDescription, SizedOutfitSource.NumResizingInterpolationPoints, InterpData);
 					}
 					InterpolationData_SampleIndices.Emplace(MoveTemp(InterpData.SampleIndices));
@@ -1222,3 +1231,5 @@ void UChaosOutfit::Merge(
 		OutSkeletalMeshRenderData->bSupportRayTracing = OutSkeletalMeshRenderData->bSupportRayTracing || InSkeletalMeshRenderData->bSupportRayTracing;
 	}
 }
+
+#undef LOCTEXT_NAMESPACE
