Chaos Outfit Asset - Added Chest as an alias for the Bust measurement.

#rnx

diff --git a/Engine/Plugins/Experimental/ChaosOutfitAsset/Source/ChaosOutfitAssetEngine/Private/ChaosOutfitAsset/CollectionOutfitFacade.cpp b/Engine/Plugins/Experimental/ChaosOutfitAsset/Source/ChaosOutfitAssetEngine/Private/ChaosOutfitAsset/CollectionOutfitFacade.cpp
--- a/Engine/Plugins/Experimental/ChaosOutfitAsset/Source/ChaosOutfitAssetEngine/Private/ChaosOutfitAsset/CollectionOutfitFacade.cpp
+++ b/Engine/Plugins/Experimental/ChaosOutfitAsset/Source/ChaosOutfitAssetEngine/Private/ChaosOutfitAsset/CollectionOutfitFacade.cpp
@@ -187,7 +187,11 @@ namespace UE::Chaos::OutfitAsset
 			explicit FBodyMatchParameters(const TMap<FString, float>& Measurements)
 			{
 				Hip = Measurements.FindRef(TEXT("Hip"), InvalidMeasurement);
-				const float Bust = Measurements.FindRef(TEXT("Bust"), InvalidMeasurement);
+				float Bust = Measurements.FindRef(TEXT("Bust"), InvalidMeasurement);
+				if (Bust == InvalidMeasurement)
+				{
+					Bust = Measurements.FindRef(TEXT("Chest"), InvalidMeasurement);  // Bust has been replaced by Chest after the 5.6 preview 1 and before the 5.6 release
+				}
 				const float Underbust = Measurements.FindRef(TEXT("Underbust"), InvalidMeasurement);
 				const float Waist = Measurements.FindRef(TEXT("Waist"), InvalidMeasurement);
 				const float NeckToWaist = Measurements.FindRef(TEXT("Neck to Waist"), InvalidMeasurement);
