CIS - Fixed Issue 904292: Compile errors in Module.ChaosOutfitAssetEngine.cpp and OutfitAsset.cpp (overzealous mobile platforms build error).

#rnx

diff --git a/Engine/Plugins/Experimental/ChaosOutfitAsset/Source/ChaosOutfitAssetEngine/Private/ChaosOutfitAsset/OutfitAsset.cpp b/Engine/Plugins/Experimental/ChaosOutfitAsset/Source/ChaosOutfitAssetEngine/Private/ChaosOutfitAsset/OutfitAsset.cpp
--- a/Engine/Plugins/Experimental/ChaosOutfitAsset/Source/ChaosOutfitAssetEngine/Private/ChaosOutfitAsset/OutfitAsset.cpp
+++ b/Engine/Plugins/Experimental/ChaosOutfitAsset/Source/ChaosOutfitAssetEngine/Private/ChaosOutfitAsset/OutfitAsset.cpp
@@ -101,10 +101,10 @@ void UChaosOutfitAsset::Build(const TObjectPtr<const UChaosOutfit> InOutfit, UE:
 			// Fix up cloth collections
 			if (MaterialPathsToFixUp.Num())
 			{
-				auto FixUpPiecesMaterials = [&MaterialPathsToFixUp](TArrayView<FChaosOutfitPiece> Pieces)
+				auto FixUpPiecesMaterials = [&MaterialPathsToFixUp](TArrayView<FChaosOutfitPiece> PiecesToFixUp)
 					{
 						using namespace UE::Chaos::ClothAsset;
-						for (FChaosOutfitPiece& Piece : Pieces)
+						for (FChaosOutfitPiece& Piece : PiecesToFixUp)
 						{
 							for (TSharedRef<const FManagedArrayCollection>& Collection : Piece.Collections)
 							{
