Chaos Outfit Asset - Fixed the filter node edit condition that is preventing setting up a size to override the target body.

#rnx

diff --git a/Engine/Plugins/Experimental/ChaosOutfitAsset/Source/ChaosOutfitAssetDataflowNodes/Private/ChaosOutfitAsset/FilterSizedOutfitNode.h b/Engine/Plugins/Experimental/ChaosOutfitAsset/Source/ChaosOutfitAssetDataflowNodes/Private/ChaosOutfitAsset/FilterSizedOutfitNode.h
--- a/Engine/Plugins/Experimental/ChaosOutfitAsset/Source/ChaosOutfitAssetDataflowNodes/Private/ChaosOutfitAsset/FilterSizedOutfitNode.h
+++ b/Engine/Plugins/Experimental/ChaosOutfitAsset/Source/ChaosOutfitAssetDataflowNodes/Private/ChaosOutfitAsset/FilterSizedOutfitNode.h
@@ -38,7 +38,7 @@ private:
 	 * The name of the body size to use to filter.
 	 * If the input Size Name is empty, the output will be set to a name set from the Target Body's measurements.
 	 */
-	UPROPERTY(EditAnywhere, Category = "FilterSizedOutfit", Meta = (DataflowInput, DataflowOutput, DataflowPassthrough = "SizeName", EditCondition = "TargetBody == nullptr"))
+	UPROPERTY(EditAnywhere, Category = "FilterSizedOutfit", Meta = (DataflowInput, DataflowOutput, DataflowPassthrough = "SizeName"))
 	FString SizeName;
 
 	/**
