Optimize ChaosOutfitAssetFilterSizedOutfit Node
To  Merge Render Data now every LOD is computed in a separated thread.
Merge Render Data went from 31ms to 19ms
And so ChaosOutfitAssetFilterSizedOutfit Node went from 55ms to 43ms.

diff --git a/Engine/Plugins/Experimental/ChaosOutfitAsset/Source/ChaosOutfitAssetEngine/Private/ChaosOutfitAsset/Outfit.cpp b/Engine/Plugins/Experimental/ChaosOutfitAsset/Source/ChaosOutfitAssetEngine/Private/ChaosOutfitAsset/Outfit.cpp
--- a/Engine/Plugins/Experimental/ChaosOutfitAsset/Source/ChaosOutfitAssetEngine/Private/ChaosOutfitAsset/Outfit.cpp
+++ b/Engine/Plugins/Experimental/ChaosOutfitAsset/Source/ChaosOutfitAssetEngine/Private/ChaosOutfitAsset/Outfit.cpp
@@ -18,6 +18,8 @@
 #include "Rendering/SkeletalMeshLODRenderData.h"
 #include "Rendering/SkeletalMeshRenderData.h"
 #include "StaticMeshResources.h"
+#include "Tasks/Task.h"
+#include "Async/TaskGraphInterfaces.h"
 
 #include UE_INLINE_GENERATED_CPP_BY_NAME(Outfit)
 
@@ -1217,6 +1219,33 @@ void UChaosOutfit::Merge(
 			{
 				LODRenderData.Add(new LODRenderDataElementType());
 			}
+		}
+
+		if (NumLODs > 1)
+		{
+
+			TArray<UE::Tasks::FTask> PendingTasks;
+			PendingTasks.Reserve(NumLODs);
+			for (int32 LODIndex = 0; LODIndex < NumLODs; ++LODIndex)
+			{
+				UE::Tasks::FTask PendingTask = UE::Tasks::Launch(UE_SOURCE_LOCATION, [LODIndex , &InLODRenderData, &AssetGuids, MaterialOffset, &OutReferenceSkeleton, &BoneMap, &LODRenderData]()
+				{
+					MergeLODRenderDatas(
+						InLODRenderData[LODIndex],
+						AssetGuids,
+						MaterialOffset,
+						OutReferenceSkeleton,
+						BoneMap,
+						LODRenderData[LODIndex]);
+				});
+				PendingTasks.Add(PendingTask);
+			}
+			UE::Tasks::Wait(PendingTasks);
+		}
+		else if (NumLODs == 1)
+		{
+			const int32 LODIndex = 0;
+
 			MergeLODRenderDatas(
 				InLODRenderData[LODIndex],
 				AssetGuids,
