Fix crash nullptr Material
Check nullptr

diff --git a/Engine/Plugins/Experimental/ChaosOutfitAsset/Source/ChaosOutfitAssetEngine/Private/ChaosOutfitAsset/OutfitAsset.cpp b/Engine/Plugins/Experimental/ChaosOutfitAsset/Source/ChaosOutfitAssetEngine/Private/ChaosOutfitAsset/OutfitAsset.cpp
--- a/Engine/Plugins/Experimental/ChaosOutfitAsset/Source/ChaosOutfitAssetEngine/Private/ChaosOutfitAsset/OutfitAsset.cpp
+++ b/Engine/Plugins/Experimental/ChaosOutfitAsset/Source/ChaosOutfitAssetEngine/Private/ChaosOutfitAsset/OutfitAsset.cpp
@@ -84,7 +84,7 @@ void UChaosOutfitAsset::Build(const TObjectPtr<const UChaosOutfit> InOutfit, UE:
 			for (int32 MaterialIndex = 0; MaterialIndex < Materials.Num(); ++MaterialIndex)
 			{
 				FSkeletalMaterial& Material = Materials[MaterialIndex];
-				if (Material.MaterialInterface->GetOuter() == GetTransientPackage())
+				if (Material.MaterialInterface && Material.MaterialInterface->GetOuter() == GetTransientPackage())
 				{
 					const FString TransientPathName = Material.MaterialInterface->GetPathName();
 					Material.MaterialInterface = Cast<UMaterialInterface>(ContextAssetStore->CommitAsset(TransientPathName));
