CIS - Fixed Warning: Text conflict from LOCTEXT macro for namespace 'OutfitAssetFactory' and key 'SelectNoTemplateTooltip'

#rnx

diff --git a/Engine/Plugins/Experimental/ChaosOutfitAsset/Source/ChaosOutfitAssetEditor/Private/ChaosOutfitAsset/OutfitAssetFactory.cpp b/Engine/Plugins/Experimental/ChaosOutfitAsset/Source/ChaosOutfitAssetEditor/Private/ChaosOutfitAsset/OutfitAssetFactory.cpp
--- a/Engine/Plugins/Experimental/ChaosOutfitAsset/Source/ChaosOutfitAssetEditor/Private/ChaosOutfitAsset/OutfitAssetFactory.cpp
+++ b/Engine/Plugins/Experimental/ChaosOutfitAsset/Source/ChaosOutfitAssetEditor/Private/ChaosOutfitAsset/OutfitAssetFactory.cpp
@@ -47,16 +47,16 @@ UObject* UChaosOutfitAssetFactory::FactoryCreateNew(UClass* Class, UObject* Pare
 					.SetToolTipText(LOCTEXT("SelectNoTemplateTooltip", "Don't add a Dataflow to this Outfit Asset. Useful when a Dataflow already exists."))
 					.SetPrimary(false),
 				SMessageDialog::FButton(LOCTEXT("SelectEmptyTemplate", "Empty Dataflow"))
-					.SetToolTipText(LOCTEXT("SelectNoTemplateTooltip", "Add an empty Dataflow with an Outfit Asset Terminal node."))
+					.SetToolTipText(LOCTEXT("SelectEmptyTemplateTooltip", "Add an empty Dataflow with an Outfit Asset Terminal node."))
 					.SetPrimary(false),
 				SMessageDialog::FButton(LOCTEXT("SelectSimpleOutfitTemplate", "Simple Outfit"))
-					.SetToolTipText(LOCTEXT("SelectNoTemplateTooltip", "Add a Dataflow with a simple Cloth Asset aggregator graph. Allows to simulate multiple Cloth Assets from the same ChaosClothComponent."))
+					.SetToolTipText(LOCTEXT("SelectSimpleOutfitTemplateTooltip", "Add a Dataflow with a simple Cloth Asset aggregator graph. Allows to simulate multiple Cloth Assets from the same ChaosClothComponent."))
 					.SetPrimary(true),
 				SMessageDialog::FButton(LOCTEXT("SelectResizableOutfitTemplate", "Resizable Outfit"))
-					.SetToolTipText(LOCTEXT("SelectNoTemplateTooltip", "Add a Dataflow that builds a single resizable garment from multiple Cloth Assets. Body sizes will have to be provided in addition to the mutiple Cloth Assets."))
+					.SetToolTipText(LOCTEXT("SelectResizableOutfitTemplateTooltip", "Add a Dataflow that builds a single resizable garment from multiple Cloth Assets. Body sizes will have to be provided in addition to the mutiple Cloth Assets."))
 					.SetPrimary(false),
 				SMessageDialog::FButton(LOCTEXT("SelectResizingGraph", "Resizing Graph"))
-					.SetToolTipText(LOCTEXT("SelectNoTemplateTooltip", "Resizing graph to test a resizable Outfit Asset."))
+					.SetToolTipText(LOCTEXT("SelectResizingGraphTooltip", "Resizing graph to test a resizable Outfit Asset."))
 					.SetPrimary(false),
 			});
 
