Chaos Outfit Asset - Fixed crash when appending a sim+render mesh cloth asset to a sim only mesh cloth asset.

[RN] Chaos Outfit Asset - Fixed crash when appending a sim+render mesh cloth asset to a sim only mesh cloth asset.

diff --git a/Engine/Plugins/Experimental/ChaosOutfitAsset/Source/ChaosOutfitAssetEngine/Private/ChaosOutfitAsset/Outfit.cpp b/Engine/Plugins/Experimental/ChaosOutfitAsset/Source/ChaosOutfitAssetEngine/Private/ChaosOutfitAsset/Outfit.cpp
--- a/Engine/Plugins/Experimental/ChaosOutfitAsset/Source/ChaosOutfitAssetEngine/Private/ChaosOutfitAsset/Outfit.cpp
+++ b/Engine/Plugins/Experimental/ChaosOutfitAsset/Source/ChaosOutfitAssetEngine/Private/ChaosOutfitAsset/Outfit.cpp
@@ -405,8 +405,9 @@ namespace UE::Chaos::OutfitAsset::Private
 		int32 NumMappingData = 0;
 		for (const FSkelMeshRenderSection& RenderSection : RenderSections)
 		{
-			// TODO: Fix ClothMappings
-			//if (!AssetGuids || AssetGuids->Find(RenderSection.ClothingData.AssetGuid) != INDEX_NONE)
+			if (!AssetGuids ||
+				!RenderSection.ClothingData.AssetGuid.IsValid() ||
+				AssetGuids->Find(RenderSection.ClothingData.AssetGuid) != INDEX_NONE)
 			{
 				for (const TArray<FMeshToMeshVertData>& ClothMappingDataLOD : RenderSection.ClothMappingDataLODs)
 				{
@@ -417,8 +418,9 @@ namespace UE::Chaos::OutfitAsset::Private
 		MappingData.Reserve(NumMappingData);
 		for (const FSkelMeshRenderSection& RenderSection : RenderSections)
 		{
-			// TODO: Fix ClothMappings
-			//if (!AssetGuids || AssetGuids->Find(RenderSection.ClothingData.AssetGuid) != INDEX_NONE)
+			if (!AssetGuids ||
+				!RenderSection.ClothingData.AssetGuid.IsValid() ||
+				AssetGuids->Find(RenderSection.ClothingData.AssetGuid) != INDEX_NONE)
 			{
 				for (const TArray<FMeshToMeshVertData>& ClothMappingDataLOD : RenderSection.ClothMappingDataLODs)
 				{
@@ -438,7 +440,9 @@ namespace UE::Chaos::OutfitAsset::Private
 		int32 NumRenderSections = 0;
 		for (const FSkelMeshRenderSection& RenderSection : RenderSections)
 		{
-			if (!AssetGuids || AssetGuids->Find(RenderSection.ClothingData.AssetGuid) != INDEX_NONE)
+			if (!AssetGuids ||
+				!RenderSection.ClothingData.AssetGuid.IsValid() ||
+				AssetGuids->Find(RenderSection.ClothingData.AssetGuid) != INDEX_NONE)
 			{
 				++NumRenderSections;
 			}
@@ -456,7 +460,9 @@ namespace UE::Chaos::OutfitAsset::Private
 			for (int32 SectionIndex = 0; SectionIndex < RenderSections.Num(); ++SectionIndex)
 			{
 				const FSkelMeshRenderSection& RenderSection = RenderSections[SectionIndex];
-				if (!AssetGuids || AssetGuids->Find(RenderSection.ClothingData.AssetGuid) != INDEX_NONE)
+				if (!AssetGuids ||
+					!RenderSection.ClothingData.AssetGuid.IsValid() ||
+					AssetGuids->Find(RenderSection.ClothingData.AssetGuid) != INDEX_NONE)
 				{
 					FClothBufferIndexMapping& ClothBufferIndexMapping = ClothBufferIndexMappings.Add_GetRef(ClothVertexBuffer.GetClothIndexMapping()[SectionIndex]);
 					ClothBufferIndexMapping.BaseVertexIndex += VertexOffset;
