Chaos Outfit Asset - Fixed the resizing when the cloth assets' referenced source assets have been moved or renamed which is causing materials to go missing.

[RN] Chaos Outfit Asset - Fixed the resizing when the cloth assets' referenced source assets have been moved or renamed which is causing materials to go missing.

diff --git a/Engine/Plugins/ChaosClothAsset/Source/ChaosClothAssetEngine/Private/ChaosClothAsset/ClothAsset.cpp b/Engine/Plugins/ChaosClothAsset/Source/ChaosClothAssetEngine/Private/ChaosClothAsset/ClothAsset.cpp
--- a/Engine/Plugins/ChaosClothAsset/Source/ChaosClothAssetEngine/Private/ChaosClothAsset/ClothAsset.cpp
+++ b/Engine/Plugins/ChaosClothAsset/Source/ChaosClothAssetEngine/Private/ChaosClothAsset/ClothAsset.cpp
@@ -562,7 +562,7 @@ void UChaosClothAsset::Build(
 			}
 			else
 			{
-				OutMaterials.Emplace();
+				OutMaterials.Emplace();  // Must keep one material slots per render pattern, do not change this behavior without updating the fix up code in GetOutfitClothCollectionsNode.cpp
 			}
 		}
 
diff --git a/Engine/Plugins/Experimental/ChaosOutfitAsset/Source/ChaosOutfitAssetDataflowNodes/ChaosOutfitAssetDataflowNodes.Build.cs b/Engine/Plugins/Experimental/ChaosOutfitAsset/Source/ChaosOutfitAssetDataflowNodes/ChaosOutfitAssetDataflowNodes.Build.cs
--- a/Engine/Plugins/Experimental/ChaosOutfitAsset/Source/ChaosOutfitAssetDataflowNodes/ChaosOutfitAssetDataflowNodes.Build.cs
+++ b/Engine/Plugins/Experimental/ChaosOutfitAsset/Source/ChaosOutfitAssetDataflowNodes/ChaosOutfitAssetDataflowNodes.Build.cs
@@ -11,9 +11,11 @@ public class ChaosOutfitAssetDataflowNodes : ModuleRules
 		PrivateDependencyModuleNames.AddRange(
 			new string[]
 			{
+				"AssetRegistry",
 				"Core",
 				"CoreUObject",
 				"Chaos",
+				"ChaosClothAsset",
 				"ChaosClothAssetEngine",
 				"ChaosOutfitAssetEngine",
 				"DataflowCore",
diff --git a/Engine/Plugins/Experimental/ChaosOutfitAsset/Source/ChaosOutfitAssetDataflowNodes/Private/ChaosOutfitAsset/GetOutfitClothCollectionsNode.cpp b/Engine/Plugins/Experimental/ChaosOutfitAsset/Source/ChaosOutfitAssetDataflowNodes/Private/ChaosOutfitAsset/GetOutfitClothCollectionsNode.cpp
--- a/Engine/Plugins/Experimental/ChaosOutfitAsset/Source/ChaosOutfitAssetDataflowNodes/Private/ChaosOutfitAsset/GetOutfitClothCollectionsNode.cpp
+++ b/Engine/Plugins/Experimental/ChaosOutfitAsset/Source/ChaosOutfitAssetDataflowNodes/Private/ChaosOutfitAsset/GetOutfitClothCollectionsNode.cpp
@@ -1,7 +1,14 @@
 // Copyright Epic Games, Inc. All Rights Reserved.
 
 #include "ChaosOutfitAsset/GetOutfitClothCollectionsNode.h"
+#include "AssetRegistry/AssetData.h"
+#include "AssetRegistry/AssetRegistryModule.h"
 #include "ChaosOutfitAsset/Outfit.h"
+#include "ChaosClothAsset/CollectionClothFacade.h"
+#include "ChaosClothAsset/CollectionClothRenderPatternFacade.h"
+#include "Engine/SkinnedAssetCommon.h"
+#include "Materials/MaterialInterface.h"
+#include "PhysicsEngine/PhysicsAsset.h"
 
 #include UE_INLINE_GENERATED_CPP_BY_NAME(GetOutfitClothCollectionsNode)
 
@@ -20,6 +27,8 @@ FChaosGetOutfitClothCollectionsNode::FChaosGetOutfitClothCollectionsNode(const U
 
 void FChaosGetOutfitClothCollectionsNode::Evaluate(UE::Dataflow::FContext& Context, const FDataflowOutput* Out) const
 {
+	using namespace UE::Chaos::ClothAsset;
+
 	if (Out->IsA(&Outfit) || Out->IsA(&ClothCollections) || Out->IsA(&NumLods))
 	{
 		TObjectPtr<const UChaosOutfit> OutOutfit;
@@ -41,9 +50,61 @@ void FChaosGetOutfitClothCollectionsNode::Evaluate(UE::Dataflow::FContext& Conte
 
 		TArray<FManagedArrayCollection> OutClothCollections;
 		OutClothCollections.Reserve(OutNumLods * OutNumPieces);
-		for (const TSharedRef<const FManagedArrayCollection>& ClothCollection : OutOutfit->GetClothCollections(InLODIndex))
+
+		// Keep track of the starting material index for the requested LOD
+		const TArray<FSkeletalMaterial>& Materials = OutOutfit->GetMaterials();
+		int32 MaterialIndex = 0;
+		for (int32 PreLODsIndex = 0; PreLODsIndex < InLODIndex; ++PreLODsIndex)
+		{
+			for (const TSharedRef<const FManagedArrayCollection>& ClothCollection : OutOutfit->GetClothCollections(PreLODsIndex))
+			{
+				const FCollectionClothConstFacade ClothFacade(ClothCollection);
+				MaterialIndex += ClothFacade.GetNumRenderPatterns();
+			}
+		}
+
+		// Fix up the fixable object paths (in case the assets have been moved) and copy the collections
+		const FAssetRegistryModule& AssetRegistryModule = FModuleManager::LoadModuleChecked<FAssetRegistryModule>(TEXT("AssetRegistry"));
+		const TArray<TSharedRef<const FManagedArrayCollection>> SelectedClothCollections = OutOutfit->GetClothCollections(InLODIndex);
+
+		for (int32 CollectionIndex = 0; CollectionIndex < SelectedClothCollections.Num(); ++CollectionIndex)
 		{
-			OutClothCollections.Emplace(*ClothCollection);
+			const TSharedRef<const FManagedArrayCollection>& ClothCollection = SelectedClothCollections[CollectionIndex];
+			const TSharedRef<FManagedArrayCollection> OutClothCollection = MakeShared<FManagedArrayCollection>(*ClothCollection);
+			FCollectionClothFacade ClothFacade(OutClothCollection);
+
+			// Fix up physics asset paths
+			if (!ClothFacade.GetPhysicsAssetPathName().IsEmpty() &&
+				!AssetRegistryModule.Get().GetAssetByObjectPath(ClothFacade.GetPhysicsAssetPathName()).IsValid())
+			{
+				const UPhysicsAsset* const PhysicsAsset = OutOutfit->GetPieces()[CollectionIndex / OutNumLods].PhysicsAsset;
+				const FString PhysicsAssetPathName = PhysicsAsset ? PhysicsAsset->GetPathName() : FString();
+
+				Context.Warning(FString::Printf(TEXT("Can't find physics asset path [%s], will replace with [%s]"), *ClothFacade.GetPhysicsAssetPathName(), *PhysicsAssetPathName), this, Out);
+				ClothFacade.SetPhysicsAssetPathName(PhysicsAssetPathName);
+			}
+
+			// Fix up material paths
+			const int32 NumRenderPatterns = ClothFacade.GetNumRenderPatterns();
+			for (int32 PatternIndex = 0; PatternIndex < NumRenderPatterns; ++PatternIndex)
+			{
+				FCollectionClothRenderPatternFacade RenderPatternFacade = ClothFacade.GetRenderPattern(PatternIndex);
+
+				if (!RenderPatternFacade.GetRenderMaterialPathName().IsEmpty() &&
+					!AssetRegistryModule.Get().GetAssetByObjectPath(RenderPatternFacade.GetRenderMaterialPathName()).IsValid())
+				{
+					const FString MaterialPathName = Materials.IsValidIndex(MaterialIndex) && Materials[MaterialIndex].MaterialInterface ?
+						Materials[MaterialIndex].MaterialInterface->GetPathName() : FString();
+					Context.Warning(FString::Printf(TEXT("Can't find material asset path [%s], will replace with [%s]"), *RenderPatternFacade.GetRenderMaterialPathName(), *MaterialPathName), this, Out);
+					RenderPatternFacade.SetRenderMaterialPathName(MaterialPathName);
+				}
+				++MaterialIndex;
+			}
+
+			// TODO: Fix up skeletal mesh paths, can't currently be done for 5.6 due to hotfix rules.
+
+			// Copy collection
+			OutClothCollections.Emplace(MoveTemp(*OutClothCollection));
 		}
 		SetValue(Context, MoveTemp(OutClothCollections), &ClothCollections);
 		SetValue(Context, OutNumLods, &NumLods);
diff --git a/Engine/Plugins/Experimental/ChaosOutfitAsset/Source/ChaosOutfitAssetEngine/Private/ChaosOutfitAsset/Outfit.cpp b/Engine/Plugins/Experimental/ChaosOutfitAsset/Source/ChaosOutfitAssetEngine/Private/ChaosOutfitAsset/Outfit.cpp
--- a/Engine/Plugins/Experimental/ChaosOutfitAsset/Source/ChaosOutfitAssetEngine/Private/ChaosOutfitAsset/Outfit.cpp
+++ b/Engine/Plugins/Experimental/ChaosOutfitAsset/Source/ChaosOutfitAssetEngine/Private/ChaosOutfitAsset/Outfit.cpp
@@ -1200,7 +1200,7 @@ void UChaosOutfit::Merge(
 
 	// Merge materials
 	const int32 MaterialOffset = OutMaterials.Num();
-	OutMaterials.Append(InMaterials);
+	OutMaterials.Append(InMaterials);  // TODO: Merge duplicate material slots. But when doing so, don't forget to consider and update the fix up code in GetOutiftClothCollectionsNode.cpp.
 
 	// Merge render data, only for existing LODs
 	// Note: This does not create LODs for the added sections when they are missing
