Chaos Outfit Asset:
- Fixed crashes when adding two cloth assets with different root bones.
- Remove ensure that fires when no reference bones and usedbone indices are set on a cloth asset preventing the testing of the crash fix.
- Also fixes problems with the MakeOutfitNode and MakeSizedOutfitNode's UI that can mess up the input arrays and cause the pins to desync which could contribute to further crashes.

[RN] Chaos Outfit Asset - Fixed crashes when adding two cloth assets with different root bones.

diff --git a/Engine/Plugins/ChaosClothAsset/Source/ChaosClothAssetEngine/Private/ChaosClothAsset/ClothSimulationModel.cpp b/Engine/Plugins/ChaosClothAsset/Source/ChaosClothAssetEngine/Private/ChaosClothAsset/ClothSimulationModel.cpp
--- a/Engine/Plugins/ChaosClothAsset/Source/ChaosClothAssetEngine/Private/ChaosClothAsset/ClothSimulationModel.cpp
+++ b/Engine/Plugins/ChaosClothAsset/Source/ChaosClothAssetEngine/Private/ChaosClothAsset/ClothSimulationModel.cpp
@@ -494,11 +494,11 @@ CHAOSCLOTHASSETENGINE_API void FChaosClothSimulationModel::RemapBoneIndices(cons
 	constexpr int32 RootBoneIndex = 0;
 
 	// Remap the reference bone if needed
-	ReferenceBoneIndex = ensure(BoneMap.IsValidIndex(ReferenceBoneIndex)) ? BoneMap[ReferenceBoneIndex] : RootBoneIndex;
+	ReferenceBoneIndex = BoneMap.IsValidIndex(ReferenceBoneIndex) ? BoneMap[ReferenceBoneIndex] : RootBoneIndex;
 
 	// UsedBoneIndices is used to reindex from the subset index to the actual skeleton bone index, so needs updating
 	for (int32& UsedBoneIndex : UsedBoneIndices)
 	{
-		UsedBoneIndex = ensure(BoneMap.IsValidIndex(UsedBoneIndex)) ? BoneMap[UsedBoneIndex] : RootBoneIndex;
+		UsedBoneIndex = BoneMap.IsValidIndex(UsedBoneIndex) ? BoneMap[UsedBoneIndex] : RootBoneIndex;
 	}
 }
diff --git a/Engine/Plugins/Experimental/ChaosOutfitAsset/Source/ChaosOutfitAssetDataflowNodes/Private/ChaosOutfitAsset/MakeOutfitNode.h b/Engine/Plugins/Experimental/ChaosOutfitAsset/Source/ChaosOutfitAssetDataflowNodes/Private/ChaosOutfitAsset/MakeOutfitNode.h
--- a/Engine/Plugins/Experimental/ChaosOutfitAsset/Source/ChaosOutfitAssetDataflowNodes/Private/ChaosOutfitAsset/MakeOutfitNode.h
+++ b/Engine/Plugins/Experimental/ChaosOutfitAsset/Source/ChaosOutfitAssetDataflowNodes/Private/ChaosOutfitAsset/MakeOutfitNode.h
@@ -47,6 +47,6 @@ private:
 	TObjectPtr<const UChaosOutfit> Outfit;
 
 	/** The cloth assets or outfits to import into a new outfit. */
-	UPROPERTY(EditAnywhere, Category = "MakeOutfit")
+	UPROPERTY()
 	TArray<FChaosClothAssetOrArrayAnyType> ClothAssets;
 };
diff --git a/Engine/Plugins/Experimental/ChaosOutfitAsset/Source/ChaosOutfitAssetDataflowNodes/Private/ChaosOutfitAsset/MakeSizedOutfitNode.h b/Engine/Plugins/Experimental/ChaosOutfitAsset/Source/ChaosOutfitAssetDataflowNodes/Private/ChaosOutfitAsset/MakeSizedOutfitNode.h
--- a/Engine/Plugins/Experimental/ChaosOutfitAsset/Source/ChaosOutfitAssetDataflowNodes/Private/ChaosOutfitAsset/MakeSizedOutfitNode.h
+++ b/Engine/Plugins/Experimental/ChaosOutfitAsset/Source/ChaosOutfitAssetDataflowNodes/Private/ChaosOutfitAsset/MakeSizedOutfitNode.h
@@ -49,7 +49,7 @@ private:
 	UPROPERTY(Meta = (DataflowOutput))
 	FManagedArrayCollection OutfitCollection;
 
-	/** The cloth assets/outfits and bidy references to import into as a new sized outfit. */
-	UPROPERTY(EditAnywhere, Category = "MakeSizedOutfit")
+	/** The cloth assets/outfits and body references to import into as a new sized outfit. */
+	UPROPERTY()
 	TArray<FChaosSizedOutfitSourceOrArrayAnyType> SizedOutfitSources;
 };
diff --git a/Engine/Plugins/Experimental/ChaosOutfitAsset/Source/ChaosOutfitAssetEngine/Private/ChaosOutfitAsset/Outfit.cpp b/Engine/Plugins/Experimental/ChaosOutfitAsset/Source/ChaosOutfitAssetEngine/Private/ChaosOutfitAsset/Outfit.cpp
--- a/Engine/Plugins/Experimental/ChaosOutfitAsset/Source/ChaosOutfitAssetEngine/Private/ChaosOutfitAsset/Outfit.cpp
+++ b/Engine/Plugins/Experimental/ChaosOutfitAsset/Source/ChaosOutfitAssetEngine/Private/ChaosOutfitAsset/Outfit.cpp
@@ -137,10 +137,11 @@ namespace UE::Chaos::OutfitAsset::Private
 		const FName BoneName = InSkeleton.GetBoneName(BoneIndex);
 		const int32 ParentBoneIndex = InSkeleton.GetParentIndex(BoneIndex);
 		const int32 ExistingBoneIndex = InOutSkeleton.FindBoneIndex(BoneName);
-		if (ExistingBoneIndex != INDEX_NONE)  // Already exists, update bone map
+		if (ExistingBoneIndex != INDEX_NONE ||  // Already exists or
+			ParentBoneIndex == INDEX_NONE)      // it is the root bone, update the bone map
 		{
 			// Check if the parents matches
-			const int32 ExistingParentBoneIndex = InOutSkeleton.GetParentIndex(ExistingBoneIndex);
+			const int32 ExistingParentBoneIndex = (ExistingBoneIndex != INDEX_NONE) ? InOutSkeleton.GetParentIndex(ExistingBoneIndex) : INDEX_NONE;
 			if (!ensureMsgf((ParentBoneIndex == INDEX_NONE && ExistingParentBoneIndex == INDEX_NONE) ||
 				(InSkeleton.GetBoneName(ParentBoneIndex) == InOutSkeleton.GetBoneName(ExistingParentBoneIndex)),
 				TEXT("Bone parent must match")))
