Chaos Outfit Asset - Fixed an issue with the body filtering that would disproportionately choose non protruding clothing from protruding bodies, especially when these bodies were smaller sized.

[RN] MetaHuman Creator/Chaos Outfit Asset - Fixed an issue with the body filtering that would disproportionately choose non protruding clothing from protruding bodies, especially when these bodies were smaller sized.

diff --git a/Engine/Plugins/Experimental/ChaosOutfitAsset/Source/ChaosOutfitAssetEngine/Private/ChaosOutfitAsset/CollectionOutfitFacade.cpp b/Engine/Plugins/Experimental/ChaosOutfitAsset/Source/ChaosOutfitAssetEngine/Private/ChaosOutfitAsset/CollectionOutfitFacade.cpp
--- a/Engine/Plugins/Experimental/ChaosOutfitAsset/Source/ChaosOutfitAssetEngine/Private/ChaosOutfitAsset/CollectionOutfitFacade.cpp
+++ b/Engine/Plugins/Experimental/ChaosOutfitAsset/Source/ChaosOutfitAssetEngine/Private/ChaosOutfitAsset/CollectionOutfitFacade.cpp
@@ -208,7 +208,11 @@ namespace UE::Chaos::OutfitAsset
 					true)
 				{
 					// Protruding
-					bIsProtruding = (Bust - Underbust) >= 15.f;
+					constexpr float MaxProtrudingRatio = 0.88f;
+					bIsProtruding = (Underbust / Bust) <= MaxProtrudingRatio;
+					UE_LOG(LogChaosOutfitAsset, Verbose,
+						TEXT("BodyMatchParameters: Bust = %.3f, Underbust = %.3f, ProtrudingRatio = %.3f, bIsProtruding = %d"),
+						Bust, Underbust, Underbust / Bust, (int)bIsProtruding);
 
 					// Torso shape
 					const float HipBust = Hip / Bust;
diff --git a/Engine/Plugins/Experimental/ChaosOutfitAsset/Source/ChaosOutfitAssetEngine/Private/ChaosOutfitAsset/Outfit.cpp b/Engine/Plugins/Experimental/ChaosOutfitAsset/Source/ChaosOutfitAssetEngine/Private/ChaosOutfitAsset/Outfit.cpp
--- a/Engine/Plugins/Experimental/ChaosOutfitAsset/Source/ChaosOutfitAssetEngine/Private/ChaosOutfitAsset/Outfit.cpp
+++ b/Engine/Plugins/Experimental/ChaosOutfitAsset/Source/ChaosOutfitAssetEngine/Private/ChaosOutfitAsset/Outfit.cpp
@@ -206,7 +206,7 @@ namespace UE::Chaos::OutfitAsset::Private
 		{
 			const FString IndentedName = IndentString + MeshBoneInfos[ChildIndex].Name.ToString();
 
-			UE_LOG(LogChaosOutfitAsset, Verbose, TEXT("%d - %s"), ChildIndex, *IndentedName);
+			UE_LOG(LogChaosOutfitAsset, VeryVerbose, TEXT("%d - %s"), ChildIndex, *IndentedName);
 			
 			LogHierachy(MeshBoneInfos, ChildIndex, Indent + 1);
 		}
@@ -245,14 +245,14 @@ namespace UE::Chaos::OutfitAsset::Private
 			BonesToProcess[BoneIndex] = false;
 		}
 
-		UE_LOG(LogChaosOutfitAsset, Verbose, TEXT("------- Outfit Reference Skseleton Merging --------"));
+		UE_LOG(LogChaosOutfitAsset, VeryVerbose, TEXT("------- Outfit Reference Skseleton Merging --------"));
 		LogHierachy(InOutSkeleton.GetRefBoneInfo());
-		UE_LOG(LogChaosOutfitAsset, Verbose, TEXT("---------------------------------------------------"));
+		UE_LOG(LogChaosOutfitAsset, VeryVerbose, TEXT("---------------------------------------------------"));
 		for (int32 Index = 0; Index < OutBoneMap.Num(); ++Index)
 		{
-			UE_LOG(LogChaosOutfitAsset, Verbose, TEXT("%d -> %d"), Index, OutBoneMap[Index]);
+			UE_LOG(LogChaosOutfitAsset, VeryVerbose, TEXT("%d -> %d"), Index, OutBoneMap[Index]);
 		}
-		UE_LOG(LogChaosOutfitAsset, Verbose, TEXT("---------------------------------------------------"));
+		UE_LOG(LogChaosOutfitAsset, VeryVerbose, TEXT("---------------------------------------------------"));
 	}
 
 	static TArray<uint32> GetIndices(const FMultiSizeIndexContainer& MultiSizeIndexContainer, int32 VertexOffset = 0)
