Chaos Cloth Outfit - Changed the GetOutfitClothCollections node path fix up code to more reliabily work on the path names instead of the material indices.

#rnx

diff --git a/Engine/Plugins/Experimental/ChaosOutfitAsset/Source/ChaosOutfitAssetDataflowNodes/Private/ChaosOutfitAsset/GetOutfitClothCollectionsNode.cpp b/Engine/Plugins/Experimental/ChaosOutfitAsset/Source/ChaosOutfitAssetDataflowNodes/Private/ChaosOutfitAsset/GetOutfitClothCollectionsNode.cpp
--- a/Engine/Plugins/Experimental/ChaosOutfitAsset/Source/ChaosOutfitAssetDataflowNodes/Private/ChaosOutfitAsset/GetOutfitClothCollectionsNode.cpp
+++ b/Engine/Plugins/Experimental/ChaosOutfitAsset/Source/ChaosOutfitAssetDataflowNodes/Private/ChaosOutfitAsset/GetOutfitClothCollectionsNode.cpp
@@ -51,17 +51,8 @@ void FChaosGetOutfitClothCollectionsNode::Evaluate(UE::Dataflow::FContext& Conte
 		TArray<FManagedArrayCollection> OutClothCollections;
 		OutClothCollections.Reserve(OutNumLods * OutNumPieces);
 
-		// Keep track of the starting material index for the requested LOD
+		// Get the material array
 		const TArray<FSkeletalMaterial>& Materials = OutOutfit->GetMaterials();
-		int32 MaterialIndex = 0;
-		for (int32 PreLODsIndex = 0; PreLODsIndex < InLODIndex; ++PreLODsIndex)
-		{
-			for (const TSharedRef<const FManagedArrayCollection>& ClothCollection : OutOutfit->GetClothCollections(PreLODsIndex))
-			{
-				const FCollectionClothConstFacade ClothFacade(ClothCollection);
-				MaterialIndex += ClothFacade.GetNumRenderPatterns();
-			}
-		}
 
 		// Fix up the fixable object paths (in case the assets have been moved) and copy the collections
 		const FAssetRegistryModule& AssetRegistryModule = FModuleManager::LoadModuleChecked<FAssetRegistryModule>(TEXT("AssetRegistry"));
@@ -89,16 +80,43 @@ void FChaosGetOutfitClothCollectionsNode::Evaluate(UE::Dataflow::FContext& Conte
 			for (int32 PatternIndex = 0; PatternIndex < NumRenderPatterns; ++PatternIndex)
 			{
 				FCollectionClothRenderPatternFacade RenderPatternFacade = ClothFacade.GetRenderPattern(PatternIndex);
+				const FString& RenderMaterialPathName = RenderPatternFacade.GetRenderMaterialPathName();
 
-				if (!RenderPatternFacade.GetRenderMaterialPathName().IsEmpty() &&
-					!AssetRegistryModule.Get().GetAssetByObjectPath(RenderPatternFacade.GetRenderMaterialPathName()).IsValid())
+				if (!RenderMaterialPathName.IsEmpty() &&
+					!AssetRegistryModule.Get().GetAssetByObjectPath(RenderMaterialPathName).IsValid())
 				{
-					const FString MaterialPathName = Materials.IsValidIndex(MaterialIndex) && Materials[MaterialIndex].MaterialInterface ?
-						Materials[MaterialIndex].MaterialInterface->GetPathName() : FString();
-					Context.Warning(FString::Printf(TEXT("Can't find material asset path [%s], will replace with [%s]"), *RenderPatternFacade.GetRenderMaterialPathName(), *MaterialPathName), this, Out);
+					TArray<FString> RenderMaterialPathElements;
+					RenderMaterialPathName.ParseIntoArray(RenderMaterialPathElements, TEXT("/"), true);
+
+					// Find the closest end path name from all asset materials (currently there is no way to safely reference materials fromn the Outfit Piece)
+					int32 BestMatchLength = 0;
+					const FSkeletalMaterial* BestMatchMaterial = nullptr;
+					for (const FSkeletalMaterial& Material : Materials)
+					{
+						if (Material.MaterialInterface)
+						{
+							TArray<FString> MaterialPathElements;
+							Material.MaterialInterface->GetPathName().ParseIntoArray(MaterialPathElements, TEXT("/"), true);
+
+							for (int32 IndexFromTheEnd = 0; IndexFromTheEnd < FMath::Min(RenderMaterialPathElements.Num(), MaterialPathElements.Num()); ++IndexFromTheEnd)
+							{
+								if (RenderMaterialPathElements.Last(IndexFromTheEnd) != MaterialPathElements.Last(IndexFromTheEnd))
+								{
+									break;
+								}
+								if (BestMatchLength == IndexFromTheEnd)
+								{
+									++BestMatchLength;
+									BestMatchMaterial = &Material;
+								}
+							}
+						}
+					}
+					const FString MaterialPathName = BestMatchMaterial ? BestMatchMaterial->MaterialInterface->GetPathName() : FString();
+
+					Context.Warning(FString::Printf(TEXT("Can't find material asset path [%s], will replace with [%s]"), *RenderMaterialPathName, *MaterialPathName), this, Out);
 					RenderPatternFacade.SetRenderMaterialPathName(MaterialPathName);
 				}
-				++MaterialIndex;
 			}
 
 			// TODO: Fix up skeletal mesh paths, can't currently be done for 5.6 due to hotfix rules.
