Update ChaosClothAsset and ChaosOutFitAsset to use SoftObjectPaths for path names in their ManagedArrayCollections rather than Strings.

diff --git a/Engine/Plugins/ChaosClothAsset/Source/ChaosClothAsset/ChaosClothAsset.Build.cs b/Engine/Plugins/ChaosClothAsset/Source/ChaosClothAsset/ChaosClothAsset.Build.cs
--- a/Engine/Plugins/ChaosClothAsset/Source/ChaosClothAsset/ChaosClothAsset.Build.cs
+++ b/Engine/Plugins/ChaosClothAsset/Source/ChaosClothAsset/ChaosClothAsset.Build.cs
@@ -12,6 +12,7 @@ public class ChaosClothAsset : ModuleRules
 			new string[]
 			{
 				"Core",
+				"CoreUObject",
 				"GeometryCore",
 				"MeshConversion",
 				"ClothingSystemRuntimeCommon",
diff --git a/Engine/Plugins/ChaosClothAsset/Source/ChaosClothAsset/Private/ChaosClothAsset/ClothCollection.cpp b/Engine/Plugins/ChaosClothAsset/Source/ChaosClothAsset/Private/ChaosClothAsset/ClothCollection.cpp
--- a/Engine/Plugins/ChaosClothAsset/Source/ChaosClothAsset/Private/ChaosClothAsset/ClothCollection.cpp
+++ b/Engine/Plugins/ChaosClothAsset/Source/ChaosClothAsset/Private/ChaosClothAsset/ClothCollection.cpp
@@ -10,8 +10,8 @@ namespace UE::Chaos::ClothAsset::Private
 	// Lods Group
 	static const TArray<FName> LodsGroupAttributes =
 	{
-		ClothCollectionAttribute::PhysicsAssetPathName,
-		ClothCollectionAttribute::SkeletalMeshPathName,
+		ClothCollectionAttribute::PhysicsAssetSoftObjectPathName,
+		ClothCollectionAttribute::SkeletalMeshSoftObjectPathName,
 	};
 
 	// Solvers Group
@@ -71,7 +71,7 @@ namespace UE::Chaos::ClothAsset::Private
 		ClothCollectionAttribute::RenderFacesStart,
 		ClothCollectionAttribute::RenderFacesEnd,
 		ClothCollectionAttribute::RenderDeformerNumInfluences,
-		ClothCollectionAttribute::RenderMaterialPathName
+		ClothCollectionAttribute::RenderMaterialSoftObjectPathName
 	};
 
 	// Sim Faces Group
@@ -188,8 +188,8 @@ namespace UE::Chaos::ClothAsset
 		: ManagedArrayCollection(InManagedArrayCollection)
 	{
 		// LODs Group
-		PhysicsAssetPathName = ManagedArrayCollection->FindAttribute<FString>(ClothCollectionAttribute::PhysicsAssetPathName, ClothCollectionGroup::Lods);
-		SkeletalMeshPathName = ManagedArrayCollection->FindAttribute<FString>(ClothCollectionAttribute::SkeletalMeshPathName, ClothCollectionGroup::Lods);
+		PhysicsAssetSoftObjectPathName = ManagedArrayCollection->FindAttribute<FSoftObjectPath>(ClothCollectionAttribute::PhysicsAssetSoftObjectPathName, ClothCollectionGroup::Lods);
+		SkeletalMeshSoftObjectPathName = ManagedArrayCollection->FindAttribute<FSoftObjectPath>(ClothCollectionAttribute::SkeletalMeshSoftObjectPathName, ClothCollectionGroup::Lods);
 
 		// Solvers Group
 		SolverGravity = ManagedArrayCollection->FindAttribute<FVector3f>(ClothCollectionAttribute::SolverGravity, ClothCollectionGroup::Solvers);
@@ -218,7 +218,7 @@ namespace UE::Chaos::ClothAsset
 		RenderFacesStart = ManagedArrayCollection->FindAttribute<int32>(ClothCollectionAttribute::RenderFacesStart, ClothCollectionGroup::RenderPatterns);
 		RenderFacesEnd = ManagedArrayCollection->FindAttribute<int32>(ClothCollectionAttribute::RenderFacesEnd, ClothCollectionGroup::RenderPatterns);
 		RenderDeformerNumInfluences = ManagedArrayCollection->FindAttribute<int32>(ClothCollectionAttribute::RenderDeformerNumInfluences, ClothCollectionGroup::RenderPatterns);
-		RenderMaterialPathName = ManagedArrayCollection->FindAttribute<FString>(ClothCollectionAttribute::RenderMaterialPathName, ClothCollectionGroup::RenderPatterns);
+		RenderMaterialSoftObjectPathName = ManagedArrayCollection->FindAttribute<FSoftObjectPath>(ClothCollectionAttribute::RenderMaterialSoftObjectPathName, ClothCollectionGroup::RenderPatterns);
 
 		//~ Fabric Group
 		FabricBendingStiffness = ManagedArrayCollection->FindAttribute<FVector3f>(ClothCollectionAttribute::FabricBendingStiffness, ClothCollectionGroup::Fabrics);
@@ -313,8 +313,8 @@ namespace UE::Chaos::ClothAsset
 
 		return
 			// LODs Group
-			PhysicsAssetPathName &&
-			SkeletalMeshPathName &&
+			PhysicsAssetSoftObjectPathName &&
+			SkeletalMeshSoftObjectPathName &&
 
 			// Solvers Group
 			(!EnumHasAnyFlags(OptionalSchemas, EClothCollectionExtendedSchemas::Solvers) ||
@@ -344,7 +344,7 @@ namespace UE::Chaos::ClothAsset
 			RenderVerticesEnd &&
 			RenderFacesStart &&
 			RenderFacesEnd &&
-			RenderMaterialPathName &&
+			RenderMaterialSoftObjectPathName &&
 			(!EnumHasAnyFlags(OptionalSchemas, EClothCollectionExtendedSchemas::RenderDeformer) ||
 				RenderDeformerNumInfluences) &&
 
@@ -463,8 +463,8 @@ namespace UE::Chaos::ClothAsset
 		}
 
 		// LODs Group
-		PhysicsAssetPathName = &GetManagedArrayCollection()->AddAttribute<FString>(ClothCollectionAttribute::PhysicsAssetPathName, ClothCollectionGroup::Lods);
-		SkeletalMeshPathName = &GetManagedArrayCollection()->AddAttribute<FString>(ClothCollectionAttribute::SkeletalMeshPathName, ClothCollectionGroup::Lods);
+		PhysicsAssetSoftObjectPathName = &GetManagedArrayCollection()->AddAttribute<FSoftObjectPath>(ClothCollectionAttribute::PhysicsAssetSoftObjectPathName, ClothCollectionGroup::Lods);
+		SkeletalMeshSoftObjectPathName = &GetManagedArrayCollection()->AddAttribute<FSoftObjectPath>(ClothCollectionAttribute::SkeletalMeshSoftObjectPathName, ClothCollectionGroup::Lods);
 
 		// Solvers Group
 		if (EnumHasAnyFlags(OptionalSchemas, EClothCollectionExtendedSchemas::Solvers))
@@ -495,7 +495,7 @@ namespace UE::Chaos::ClothAsset
 		RenderVerticesEnd = &GetManagedArrayCollection()->AddAttribute<int32>(ClothCollectionAttribute::RenderVerticesEnd, ClothCollectionGroup::RenderPatterns, RenderVerticesDependency);
 		RenderFacesStart = &GetManagedArrayCollection()->AddAttribute<int32>(ClothCollectionAttribute::RenderFacesStart, ClothCollectionGroup::RenderPatterns, RenderFacesDependency);
 		RenderFacesEnd = &GetManagedArrayCollection()->AddAttribute<int32>(ClothCollectionAttribute::RenderFacesEnd, ClothCollectionGroup::RenderPatterns, RenderFacesDependency);
-		RenderMaterialPathName = &GetManagedArrayCollection()->AddAttribute<FString>(ClothCollectionAttribute::RenderMaterialPathName, ClothCollectionGroup::RenderPatterns);
+		RenderMaterialSoftObjectPathName = &GetManagedArrayCollection()->AddAttribute<FSoftObjectPath>(ClothCollectionAttribute::RenderMaterialSoftObjectPathName, ClothCollectionGroup::RenderPatterns);
 
 		//~ Fabric Group
 		if (EnumHasAnyFlags(OptionalSchemas, EClothCollectionExtendedSchemas::Fabrics))
@@ -585,6 +585,55 @@ namespace UE::Chaos::ClothAsset
 		}
 	}
 
+	void FClothCollection::PostSerialize(const FArchive& Ar)
+	{
+		if (Ar.IsLoading())
+		{
+			auto ConvertStringToSoftObjectPath = [](const TManagedArray<FString>& StringArray, TManagedArray<FSoftObjectPath>& SoftObjectPathArray)
+				{
+					check(StringArray.Num() == SoftObjectPathArray.Num());
+					for (int32 Index = 0; Index < StringArray.Num(); ++Index)
+					{
+						SoftObjectPathArray[Index] = FSoftObjectPath(StringArray[Index]);
+						SoftObjectPathArray[Index].PreSavePath(); // This will fixup the path with redirects.
+					}
+				};
+
+			// PhysicsAssetPathName -> PhysicsAssetSoftObjectPathName
+			if (const TManagedArray<FString>* const InPhysicsAssetPathArray = ManagedArrayCollection->FindAttributeTyped<FString>(ClothCollectionAttribute::PhysicsAssetPathName, ClothCollectionGroup::Lods))
+			{
+				if (!PhysicsAssetSoftObjectPathName)
+				{
+					PhysicsAssetSoftObjectPathName = &GetManagedArrayCollection()->AddAttribute<FSoftObjectPath>(ClothCollectionAttribute::PhysicsAssetSoftObjectPathName, ClothCollectionGroup::Lods);
+				}
+				ConvertStringToSoftObjectPath(*InPhysicsAssetPathArray, *GetPhysicsAssetSoftObjectPathName());
+				GetManagedArrayCollection()->RemoveAttribute(ClothCollectionAttribute::PhysicsAssetPathName, ClothCollectionGroup::Lods);
+			}
+
+			// SkeletalMeshPathName -> SkeletalMeshSoftObjectPathName
+			if (const TManagedArray<FString>* const InSkeletalMeshPathName = ManagedArrayCollection->FindAttributeTyped<FString>(ClothCollectionAttribute::SkeletalMeshPathName, ClothCollectionGroup::Lods))
+			{
+				if (!SkeletalMeshSoftObjectPathName)
+				{
+					SkeletalMeshSoftObjectPathName = &GetManagedArrayCollection()->AddAttribute<FSoftObjectPath>(ClothCollectionAttribute::SkeletalMeshSoftObjectPathName, ClothCollectionGroup::Lods);
+				}
+				ConvertStringToSoftObjectPath(*InSkeletalMeshPathName, *GetSkeletalMeshSoftObjectPathName());
+				GetManagedArrayCollection()->RemoveAttribute(ClothCollectionAttribute::SkeletalMeshPathName, ClothCollectionGroup::Lods);
+			}
+
+			// RenderMaterialPathName -> RenderMaterialSoftObjectPathName
+			if (const TManagedArray<FString>* const InRenderMaterialPathName = ManagedArrayCollection->FindAttributeTyped<FString>(ClothCollectionAttribute::RenderMaterialPathName, ClothCollectionGroup::RenderPatterns))
+			{
+				if (!RenderMaterialSoftObjectPathName)
+				{
+					RenderMaterialSoftObjectPathName = &GetManagedArrayCollection()->AddAttribute<FSoftObjectPath>(ClothCollectionAttribute::RenderMaterialSoftObjectPathName, ClothCollectionGroup::RenderPatterns);
+				}
+				ConvertStringToSoftObjectPath(*InRenderMaterialPathName, *GetRenderMaterialSoftObjectPathName());
+				GetManagedArrayCollection()->RemoveAttribute(ClothCollectionAttribute::RenderMaterialPathName, ClothCollectionGroup::RenderPatterns);
+			}
+		}
+	}
+
 	int32 FConstClothCollection::GetNumElements(const FName& GroupName) const
 	{
 		return ManagedArrayCollection->NumElements(GroupName);
diff --git a/Engine/Plugins/ChaosClothAsset/Source/ChaosClothAsset/Private/ChaosClothAsset/ClothCollection.h b/Engine/Plugins/ChaosClothAsset/Source/ChaosClothAsset/Private/ChaosClothAsset/ClothCollection.h
--- a/Engine/Plugins/ChaosClothAsset/Source/ChaosClothAsset/Private/ChaosClothAsset/ClothCollection.h
+++ b/Engine/Plugins/ChaosClothAsset/Source/ChaosClothAsset/Private/ChaosClothAsset/ClothCollection.h
@@ -130,8 +130,12 @@ namespace UE::Chaos::ClothAsset
 		bool IsValidUserDefinedAttributeName(const FName& Name, const FName& GroupName) const;
 
 		//~ LODs Group (There should be only one LOD per ClothCollection)
-		const TManagedArray<FString>* GetPhysicsAssetPathName() const { return PhysicsAssetPathName; }
-		const TManagedArray<FString>* GetSkeletalMeshPathName() const { return SkeletalMeshPathName; }
+		const TManagedArray<FSoftObjectPath>* GetPhysicsAssetSoftObjectPathName() const { return PhysicsAssetSoftObjectPathName; }
+		const TManagedArray<FSoftObjectPath>* GetSkeletalMeshSoftObjectPathName() const { return SkeletalMeshSoftObjectPathName; }
+		UE_DEPRECATED(5.7, "Use GetPhysicsAssetSoftObjectPathName instead")
+		const TManagedArray<FString>* GetPhysicsAssetPathName() const { return nullptr; }
+		UE_DEPRECATED(5.7, "Use GetSkeletalMeshSoftObjectPathName instead")
+		const TManagedArray<FString>* GetSkeletalMeshPathName() const { return nullptr; }
 
 		// ~ Solvers Group
 		const TManagedArray<FVector3f>* GetSolverGravity() const { return SolverGravity; }
@@ -172,7 +176,9 @@ namespace UE::Chaos::ClothAsset
 		const TManagedArray<int32>* GetRenderFacesStart() const { return RenderFacesStart; }
 		const TManagedArray<int32>* GetRenderFacesEnd() const { return RenderFacesEnd; }
 		const TManagedArray<int32>* GetRenderDeformerNumInfluences() const { return RenderDeformerNumInfluences; }
-		const TManagedArray<FString>* GetRenderMaterialPathName() const { return RenderMaterialPathName; }
+		const TManagedArray<FSoftObjectPath>* GetRenderMaterialSoftObjectPathName() const { return RenderMaterialSoftObjectPathName; }
+		UE_DEPRECATED(5.7, "Use GetRenderMaterialSoftObjectPathName instead")
+		const TManagedArray<FString>* GetRenderMaterialPathName() const { return nullptr; }
 
 		//~ Sim Faces Group
 		const TManagedArray<FIntVector3>* GetSimIndices2D() const { return SimIndices2D; }
@@ -238,7 +244,11 @@ namespace UE::Chaos::ClothAsset
 		TSharedRef<const FManagedArrayCollection> ManagedArrayCollection;
 
 		//~ LODs Group
+		const TManagedArray<FSoftObjectPath>* PhysicsAssetSoftObjectPathName;
+		const TManagedArray<FSoftObjectPath>* SkeletalMeshSoftObjectPathName;
+		UE_DEPRECATED(5.7, "Use PhysicsAssetSoftObjectPathName instead")
 		const TManagedArray<FString>* PhysicsAssetPathName;
+		UE_DEPRECATED(5.7, "Use SkeletalMeshSoftObjectPathName instead")
 		const TManagedArray<FString>* SkeletalMeshPathName;
 
 		//~ Solvers group
@@ -280,6 +290,8 @@ namespace UE::Chaos::ClothAsset
 		const TManagedArray<int32>* RenderFacesStart;
 		const TManagedArray<int32>* RenderFacesEnd;
 		const TManagedArray<int32>* RenderDeformerNumInfluences;  // Number of deformer mapping influences per render vertex, either 0 (no mappings), 1 or 5 (= NUM_INFLUENCES_PER_VERTEX as defined in the ush files)
+		const TManagedArray<FSoftObjectPath>* RenderMaterialSoftObjectPathName;
+		UE_DEPRECATED(5.7, "Use RenderMaterialSoftObjectPathName instead")
 		const TManagedArray<FString>* RenderMaterialPathName;
 
 		//~ Sim Faces Group
@@ -354,6 +366,9 @@ namespace UE::Chaos::ClothAsset
 		/** Make the underlying collection a cloth collection. */
 		void DefineSchema(EClothCollectionExtendedSchemas OptionalSchemas = EClothCollectionExtendedSchemas::None);
 
+		/** Do any PostSerialize fixups */
+		void PostSerialize(const FArchive& Ar);
+
 		/** Set the number of elements of a group. */
 		void SetNumElements(int32 InNumElements, const FName& GroupName);
 
@@ -381,8 +396,12 @@ namespace UE::Chaos::ClothAsset
 		TManagedArray<T>* GetUserDefinedAttribute(const FName& Name, const FName& GroupName);
 
 		//~ LODs Group (There should be only one LOD per ClothCollection)
-		TManagedArray<FString>* GetPhysicsAssetPathName() { return const_cast<TManagedArray<FString>*>(PhysicsAssetPathName); }
-		TManagedArray<FString>* GetSkeletalMeshPathName() { return const_cast<TManagedArray<FString>*>(SkeletalMeshPathName); }
+		TManagedArray<FSoftObjectPath>* GetPhysicsAssetSoftObjectPathName() { return  const_cast<TManagedArray<FSoftObjectPath>*>(PhysicsAssetSoftObjectPathName); }
+		TManagedArray<FSoftObjectPath>* GetSkeletalMeshSoftObjectPathName() { return  const_cast<TManagedArray<FSoftObjectPath>*>(SkeletalMeshSoftObjectPathName); }
+		UE_DEPRECATED(5.7, "Use GetPhysicsAssetSoftObjectPathName instead")
+		TManagedArray<FString>* GetPhysicsAssetPathName() { return nullptr; }
+		UE_DEPRECATED(5.7, "Use GetSkeletalMeshSoftObjectPathName instead")
+		TManagedArray<FString>* GetSkeletalMeshPathName() { return nullptr; }
 
 		//~ Solver Group
 		TManagedArray<FVector3f>* GetSolverGravity() { return const_cast<TManagedArray<FVector3f>*>(SolverGravity); }
@@ -423,7 +442,9 @@ namespace UE::Chaos::ClothAsset
 		TManagedArray<int32>* GetRenderFacesStart() { return const_cast<TManagedArray<int32>*>(RenderFacesStart); }
 		TManagedArray<int32>* GetRenderFacesEnd() { return const_cast<TManagedArray<int32>*>(RenderFacesEnd); }
 		TManagedArray<int32>* GetRenderDeformerNumInfluences() { return const_cast<TManagedArray<int32>*>(RenderDeformerNumInfluences); }
-		TManagedArray<FString>* GetRenderMaterialPathName() { return const_cast<TManagedArray<FString>*>(RenderMaterialPathName); }
+		TManagedArray<FSoftObjectPath>* GetRenderMaterialSoftObjectPathName() { return const_cast<TManagedArray<FSoftObjectPath>*>(RenderMaterialSoftObjectPathName); }
+		UE_DEPRECATED(5.7, "Use GetRenderMaterialSoftObjectPathName instead")
+		TManagedArray<FString>* GetRenderMaterialPathName() { return nullptr; }
 
 		//~ Sim Faces Group
 		TManagedArray<FIntVector3>* GetSimIndices2D() { return const_cast<TManagedArray<FIntVector3>*>(SimIndices2D); }
diff --git a/Engine/Plugins/ChaosClothAsset/Source/ChaosClothAsset/Private/ChaosClothAsset/ClothGeometryTools.cpp b/Engine/Plugins/ChaosClothAsset/Source/ChaosClothAsset/Private/ChaosClothAsset/ClothGeometryTools.cpp
--- a/Engine/Plugins/ChaosClothAsset/Source/ChaosClothAsset/Private/ChaosClothAsset/ClothGeometryTools.cpp
+++ b/Engine/Plugins/ChaosClothAsset/Source/ChaosClothAsset/Private/ChaosClothAsset/ClothGeometryTools.cpp
@@ -522,7 +522,7 @@ namespace UE::Chaos::ClothAsset
 		}
 	}
 
-	void FClothGeometryTools::CopySimMeshToRenderMesh(const TSharedRef<FManagedArrayCollection>& ClothCollection, const FString& RenderMaterialPathName, bool bSingleRenderPattern)
+	void FClothGeometryTools::CopySimMeshToRenderMesh(const TSharedRef<FManagedArrayCollection>& ClothCollection, const FSoftObjectPath& RenderMaterialPathName, bool bSingleRenderPattern)
 	{
 		FCollectionClothFacade ClothFacade(ClothCollection);
 
@@ -536,7 +536,7 @@ namespace UE::Chaos::ClothAsset
 		for (int32 RenderPatternIndex = 0; RenderPatternIndex < NumRenderPatterns; ++RenderPatternIndex)
 		{
 			FCollectionClothRenderPatternFacade RenderPattern = ClothFacade.GetRenderPattern(RenderPatternIndex);
-			RenderPattern.SetRenderMaterialPathName(RenderMaterialPathName);
+			RenderPattern.SetRenderMaterialSoftObjectPathName(RenderMaterialPathName);
 			RenderPattern.SetNumRenderVertices(bSingleRenderPattern ? TotalNumVertices : ClothFacade.GetSimPattern(RenderPatternIndex).GetNumSimVertices2D());
 			RenderPattern.SetNumRenderFaces(bSingleRenderPattern ? TotalNumFaces : ClothFacade.GetSimPattern(RenderPatternIndex).GetNumSimFaces());
 		}
diff --git a/Engine/Plugins/ChaosClothAsset/Source/ChaosClothAsset/Private/ChaosClothAsset/CollectionClothFacade.cpp b/Engine/Plugins/ChaosClothAsset/Source/ChaosClothAsset/Private/ChaosClothAsset/CollectionClothFacade.cpp
--- a/Engine/Plugins/ChaosClothAsset/Source/ChaosClothAsset/Private/ChaosClothAsset/CollectionClothFacade.cpp
+++ b/Engine/Plugins/ChaosClothAsset/Source/ChaosClothAsset/Private/ChaosClothAsset/CollectionClothFacade.cpp
@@ -59,8 +59,8 @@ namespace UE::Chaos::ClothAsset
 		uint32 ResultHash = PreviousHash;
 
 		//~ LOD (single per collection) Group
-		ResultHash = HashCombineFast(ResultHash, GetTypeHash(GetSkeletalMeshPathName()));
-		ResultHash = HashCombineFast(ResultHash, GetTypeHash(GetPhysicsAssetPathName()));
+		ResultHash = HashCombineFast(ResultHash, GetTypeHash(GetPhysicsAssetSoftObjectPathName()));
+		ResultHash = HashCombineFast(ResultHash, GetTypeHash(GetSkeletalMeshSoftObjectPathName()));
 
 		//~ Solvers Group
 		ResultHash = HashCombineFast(ResultHash, GetTypeHash(GetSolverGravity()));
@@ -88,7 +88,7 @@ namespace UE::Chaos::ClothAsset
 		ResultHash = HashCombineFast(ResultHash, ClothCollection->GetElementsTypeHash(ClothCollection->GetRenderVerticesEnd()));
 		ResultHash = HashCombineFast(ResultHash, ClothCollection->GetElementsTypeHash(ClothCollection->GetRenderFacesStart()));
 		ResultHash = HashCombineFast(ResultHash, ClothCollection->GetElementsTypeHash(ClothCollection->GetRenderFacesEnd()));
-		ResultHash = HashCombineFast(ResultHash, ClothCollection->GetElementsTypeHash(ClothCollection->GetRenderMaterialPathName()));
+		ResultHash = HashCombineFast(ResultHash, ClothCollection->GetElementsTypeHash(ClothCollection->GetRenderMaterialSoftObjectPathName()));
 		ResultHash = HashCombineFast(ResultHash, ClothCollection->GetElementsTypeHash(ClothCollection->GetRenderDeformerNumInfluences()));
 
 		//~ Sim Faces Group
@@ -193,16 +193,16 @@ namespace UE::Chaos::ClothAsset
 	template CHAOSCLOTHASSET_API uint32 FCollectionClothConstFacade::CalculateUserDefinedAttributesTypeHash<FVector3f>(const FName& GroupName, uint32 PreviousHash) const;
 	template CHAOSCLOTHASSET_API uint32 FCollectionClothConstFacade::CalculateUserDefinedAttributesTypeHash<TArray<int32>>(const FName& GroupName, uint32 PreviousHash) const;
 
-	const FString& FCollectionClothConstFacade::GetPhysicsAssetPathName() const
+	const FSoftObjectPath& FCollectionClothConstFacade::GetPhysicsAssetSoftObjectPathName() const
 	{
-		static const FString EmptyString;
-		return ClothCollection->GetPhysicsAssetPathName() && ClothCollection->GetNumElements(ClothCollectionGroup::Lods) > 0 ? (*ClothCollection->GetPhysicsAssetPathName())[0] : EmptyString;
+		static const FSoftObjectPath EmptyString;
+		return ClothCollection->GetPhysicsAssetSoftObjectPathName() && ClothCollection->GetNumElements(ClothCollectionGroup::Lods) > 0 ? (*ClothCollection->GetPhysicsAssetSoftObjectPathName())[0] : EmptyString;
 	}
 
-	const FString& FCollectionClothConstFacade::GetSkeletalMeshPathName() const
+	const FSoftObjectPath& FCollectionClothConstFacade::GetSkeletalMeshSoftObjectPathName() const
 	{
-		static const FString EmptyString;
-		return ClothCollection->GetSkeletalMeshPathName() && ClothCollection->GetNumElements(ClothCollectionGroup::Lods) > 0 ? (*ClothCollection->GetSkeletalMeshPathName())[0] : EmptyString;
+		static const FSoftObjectPath EmptyString;
+		return ClothCollection->GetSkeletalMeshSoftObjectPathName() && ClothCollection->GetNumElements(ClothCollectionGroup::Lods) > 0 ? (*ClothCollection->GetSkeletalMeshSoftObjectPathName())[0] : EmptyString;
 	}
 
 	bool FCollectionClothConstFacade::HasSolverElement() const
@@ -398,9 +398,9 @@ namespace UE::Chaos::ClothAsset
 		return ClothCollection->GetElements(ClothCollection->GetRenderDeformerNumInfluences());
 	}
 
-	TConstArrayView<FString> FCollectionClothConstFacade::GetRenderMaterialPathName() const
+	TConstArrayView<FSoftObjectPath> FCollectionClothConstFacade::GetRenderMaterialSoftObjectPathName() const
 	{
-		return ClothCollection->GetElements(ClothCollection->GetRenderMaterialPathName());
+		return ClothCollection->GetElements(ClothCollection->GetRenderMaterialSoftObjectPathName());
 	}
 
 	int32 FCollectionClothConstFacade::FindRenderPatternByVertex(int32 VertexIndex) const
@@ -650,6 +650,11 @@ namespace UE::Chaos::ClothAsset
 		SetDefaults();
 	}
 
+	void FCollectionClothFacade::PostSerialize(const FArchive& Ar)
+	{
+		GetClothCollection()->PostSerialize(Ar);
+	}
+
 	void FCollectionClothFacade::Reset()
 	{
 		GetClothCollection()->SetNumElements(0, ClothCollectionGroup::Lods);  // Reset all Lods data
@@ -696,13 +701,13 @@ namespace UE::Chaos::ClothAsset
 
 		// LODs Group 
 		// Keep original data unless our data is empty. Then take Other's data.
-		if (GetPhysicsAssetPathName().IsEmpty())
+		if (GetPhysicsAssetSoftObjectPathName().IsNull())
 		{
-			SetPhysicsAssetPathName(Other.GetPhysicsAssetPathName());
+			SetPhysicsAssetSoftObjectPathName(Other.GetPhysicsAssetSoftObjectPathName());
 		}
-		if (GetSkeletalMeshPathName().IsEmpty())
+		if (GetSkeletalMeshSoftObjectPathName().IsNull())
 		{
-			SetSkeletalMeshPathName(Other.GetSkeletalMeshPathName());
+			SetSkeletalMeshSoftObjectPathName(Other.GetSkeletalMeshSoftObjectPathName());
 		}
 
 		const bool bOtherHasResizingData = Other.IsValid(EClothCollectionExtendedSchemas::Resizing);
@@ -865,19 +870,19 @@ namespace UE::Chaos::ClothAsset
 		FClothCollection::CopyArrayViewData(GetPreResizedSimPosition3D().Right(Other.GetNumSimVertices3D()), Other.GetPreResizedSimPosition3D());
 	}
 
-	void FCollectionClothFacade::SetPhysicsAssetPathName(const FString& PathName)
+	void FCollectionClothFacade::SetPhysicsAssetSoftObjectPathName(const FSoftObjectPath& PathName)
 	{
-		if (ClothCollection->GetNumElements(ClothCollectionGroup::Lods))
+		if (ClothCollection->GetNumElements(ClothCollectionGroup::Lods) && GetClothCollection()->GetPhysicsAssetSoftObjectPathName())
 		{
-			(*GetClothCollection()->GetPhysicsAssetPathName())[0] = PathName;
+			(*GetClothCollection()->GetPhysicsAssetSoftObjectPathName())[0] = PathName;
 		}
 	}
 	
-	void FCollectionClothFacade::SetSkeletalMeshPathName(const FString& PathName)
+	void FCollectionClothFacade::SetSkeletalMeshSoftObjectPathName(const FSoftObjectPath& PathName)
 	{
-		if (ClothCollection->GetNumElements(ClothCollectionGroup::Lods))
+		if (ClothCollection->GetNumElements(ClothCollectionGroup::Lods) && GetClothCollection()->GetSkeletalMeshSoftObjectPathName())
 		{
-			(*GetClothCollection()->GetSkeletalMeshPathName())[0] = PathName;
+			(*GetClothCollection()->GetSkeletalMeshSoftObjectPathName())[0] = PathName;
 		}
 	}
 	
@@ -1195,9 +1200,9 @@ namespace UE::Chaos::ClothAsset
 		return GetClothCollection()->GetElements(GetClothCollection()->GetRenderDeformerNumInfluences());
 	}
 
-	TArrayView<FString> FCollectionClothFacade::GetRenderMaterialPathName()
+	TArrayView<FSoftObjectPath> FCollectionClothFacade::GetRenderMaterialSoftObjectPathName()
 	{
-		return GetClothCollection()->GetElements(GetClothCollection()->GetRenderMaterialPathName());
+		return GetClothCollection()->GetElements(GetClothCollection()->GetRenderMaterialSoftObjectPathName());
 	}
 
 	void FCollectionClothFacade::SetNumSeams(int32 InNumSeams)
diff --git a/Engine/Plugins/ChaosClothAsset/Source/ChaosClothAsset/Private/ChaosClothAsset/CollectionClothRenderPatternFacade.cpp b/Engine/Plugins/ChaosClothAsset/Source/ChaosClothAsset/Private/ChaosClothAsset/CollectionClothRenderPatternFacade.cpp
--- a/Engine/Plugins/ChaosClothAsset/Source/ChaosClothAsset/Private/ChaosClothAsset/CollectionClothRenderPatternFacade.cpp
+++ b/Engine/Plugins/ChaosClothAsset/Source/ChaosClothAsset/Private/ChaosClothAsset/CollectionClothRenderPatternFacade.cpp
@@ -13,11 +13,11 @@ namespace UE::Chaos::ClothAsset
 			(*ClothCollection->GetRenderDeformerNumInfluences())[GetElementIndex()] : 0;
 	}
 
-	const FString& FCollectionClothRenderPatternConstFacade::GetRenderMaterialPathName() const
+	const FSoftObjectPath& FCollectionClothRenderPatternConstFacade::GetRenderMaterialSoftObjectPathName() const
 	{
-		static const FString EmptyString;
-		return ClothCollection->GetRenderMaterialPathName() && ClothCollection->GetNumElements(ClothCollectionGroup::RenderPatterns) > GetElementIndex() ?
-			(*ClothCollection->GetRenderMaterialPathName())[GetElementIndex()] : EmptyString;
+		static const FSoftObjectPath EmptyString;
+		return ClothCollection->GetRenderMaterialSoftObjectPathName() && ClothCollection->GetNumElements(ClothCollectionGroup::RenderPatterns) > GetElementIndex() ?
+			(*ClothCollection->GetRenderMaterialSoftObjectPathName())[GetElementIndex()] : EmptyString;
 	}
 
 	int32 FCollectionClothRenderPatternConstFacade::GetNumRenderVertices() const
@@ -222,7 +222,7 @@ namespace UE::Chaos::ClothAsset
 
 		//~ Render Patterns Group
 		SetRenderDeformerNumInfluences(Other.GetRenderDeformerNumInfluences());  // Must be called prior to copying the Render Deformer as it also sets the optional RenderDeformer schema
-		SetRenderMaterialPathName(Other.GetRenderMaterialPathName());
+		SetRenderMaterialSoftObjectPathName(Other.GetRenderMaterialSoftObjectPathName());
 
 		//~ Render Vertices Group
 		SetNumRenderVertices(Other.GetNumRenderVertices());
@@ -270,9 +270,9 @@ namespace UE::Chaos::ClothAsset
 		(*GetClothCollection()->GetRenderDeformerNumInfluences())[GetElementIndex()] = NumInfluences;
 	}
 
-	void FCollectionClothRenderPatternFacade::SetRenderMaterialPathName(const FString& PathName)
+	void FCollectionClothRenderPatternFacade::SetRenderMaterialSoftObjectPathName(const FSoftObjectPath& PathName)
 	{
-		(*GetClothCollection()->GetRenderMaterialPathName())[GetElementIndex()] = PathName;
+		(*GetClothCollection()->GetRenderMaterialSoftObjectPathName())[GetElementIndex()] = PathName;
 	}
 
 	void FCollectionClothRenderPatternFacade::SetNumRenderVertices(int32 NumRenderVertices)
diff --git a/Engine/Plugins/ChaosClothAsset/Source/ChaosClothAsset/Public/ChaosClothAsset/ClothCollectionAttribute.h b/Engine/Plugins/ChaosClothAsset/Source/ChaosClothAsset/Public/ChaosClothAsset/ClothCollectionAttribute.h
--- a/Engine/Plugins/ChaosClothAsset/Source/ChaosClothAsset/Public/ChaosClothAsset/ClothCollectionAttribute.h
+++ b/Engine/Plugins/ChaosClothAsset/Source/ChaosClothAsset/Public/ChaosClothAsset/ClothCollectionAttribute.h
@@ -12,8 +12,8 @@ namespace UE::Chaos::ClothAsset
 	namespace ClothCollectionAttribute
 	{
 		// Lods Group
-		inline const FName PhysicsAssetPathName(TEXT("PhysicsAssetPathName"));
-		inline const FName SkeletalMeshPathName(TEXT("SkeletalMeshPathName"));
+		inline const FName PhysicsAssetSoftObjectPathName(TEXT("PhysicsAssetSoftObjectPathName"));
+		inline const FName SkeletalMeshSoftObjectPathName(TEXT("SkeletalMeshSoftObjectPathName"));
 
 		// Solvers Group
 		inline const FName SolverGravity(TEXT("SolverGravity"));
@@ -53,7 +53,7 @@ namespace UE::Chaos::ClothAsset
 		inline const FName RenderVerticesEnd(TEXT("RenderVerticesEnd"));
 		inline const FName RenderFacesStart(TEXT("RenderFacesStart"));
 		inline const FName RenderFacesEnd(TEXT("RenderFacesEnd"));
-		inline const FName RenderMaterialPathName(TEXT("RenderMaterialPathName"));
+		inline const FName RenderMaterialSoftObjectPathName(TEXT("RenderMaterialSoftObjectPathName"));
 		inline const FName RenderDeformerNumInfluences(TEXT("RenderDeformerNumInfluences"));
 
 		// Sim Faces Group
@@ -109,5 +109,16 @@ namespace UE::Chaos::ClothAsset
 		// Custom Resizing Regions Group
 		inline const FName CustomResizingRegionSet(TEXT("CustomResizingRegionSet"));
 		inline const FName CustomResizingRegionType(TEXT("CustomResizingRegionType"));
+
+		//~ Deprecated attributes
+		// Lods Group
+		UE_DEPRECATED(5.7, "Use PhysicsAssetSoftObjectPath instead");
+		inline const FName PhysicsAssetPathName(TEXT("PhysicsAssetPathName"));
+		UE_DEPRECATED(5.7, "Use SkeletalMeshSoftObjectPathName instead");
+		inline const FName SkeletalMeshPathName(TEXT("SkeletalMeshPathName"));
+
+		// Render Patterns Group
+		UE_DEPRECATED(5.7, "Use RenderMaterialSoftObjectPathName instead");
+		inline const FName RenderMaterialPathName(TEXT("RenderMaterialPathName"));
 	}
 }  // End namespace UE::Chaos::ClothAsset
diff --git a/Engine/Plugins/ChaosClothAsset/Source/ChaosClothAsset/Public/ChaosClothAsset/ClothGeometryTools.h b/Engine/Plugins/ChaosClothAsset/Source/ChaosClothAsset/Public/ChaosClothAsset/ClothGeometryTools.h
--- a/Engine/Plugins/ChaosClothAsset/Source/ChaosClothAsset/Public/ChaosClothAsset/ClothGeometryTools.h
+++ b/Engine/Plugins/ChaosClothAsset/Source/ChaosClothAsset/Public/ChaosClothAsset/ClothGeometryTools.h
@@ -39,7 +39,7 @@ namespace UE::Chaos::ClothAsset
 		static UE_API void DeleteSelections(const TSharedRef<FManagedArrayCollection>& ClothCollection, const FName Group = NAME_None);
 
 		/** Turn the sim mesh portion of this ClothCollection into a render mesh. */
-		static UE_API void CopySimMeshToRenderMesh(const TSharedRef<FManagedArrayCollection>& ClothCollection, const FString& RenderMaterialPathName, bool bSingleRenderPattern);
+		static UE_API void CopySimMeshToRenderMesh(const TSharedRef<FManagedArrayCollection>& ClothCollection, const FSoftObjectPath& RenderMaterialPathName, bool bSingleRenderPattern);
 
 		/** Reverse the mesh normals. Will reverse all normals if pattern selection is empty. */
 		static UE_API void ReverseMesh(
diff --git a/Engine/Plugins/ChaosClothAsset/Source/ChaosClothAsset/Public/ChaosClothAsset/CollectionClothFacade.h b/Engine/Plugins/ChaosClothAsset/Source/ChaosClothAsset/Public/ChaosClothAsset/CollectionClothFacade.h
--- a/Engine/Plugins/ChaosClothAsset/Source/ChaosClothAsset/Public/ChaosClothAsset/CollectionClothFacade.h
+++ b/Engine/Plugins/ChaosClothAsset/Source/ChaosClothAsset/Public/ChaosClothAsset/CollectionClothFacade.h
@@ -69,10 +69,21 @@ namespace UE::Chaos::ClothAsset
 
 		//~ LOD (single per collection) Group
 		/** Return the physics asset path names used for this collection. */
-		UE_API const FString& GetPhysicsAssetPathName() const;
+		UE_API const FSoftObjectPath& GetPhysicsAssetSoftObjectPathName() const;
 		/** Return the skeleton asset path names used for this collection. */
-		UE_API const FString& GetSkeletalMeshPathName() const;
-		
+		UE_API const FSoftObjectPath& GetSkeletalMeshSoftObjectPathName() const;
+		UE_DEPRECATED(5.7, "Use GetPhysicsAssetSoftObjectPathName instead")
+		const FString& GetPhysicsAssetPathName() const
+		{
+			static const FString EmptyString;
+			return EmptyString;
+		}
+		UE_DEPRECATED(5.7, "Use GetSkeletalMeshSoftObjectPathName instead")
+		const FString& GetSkeletalMeshPathName() const
+		{
+			static const FString EmptyString;
+			return EmptyString;
+		}
 
 		//~ Solver (single per collection) Group
 		/** Return true if the solver group has one element*/
@@ -150,11 +161,16 @@ namespace UE::Chaos::ClothAsset
 		/** Return a view of all the render deformer number of influences used on this collection across all patterns. */
 		UE_API TConstArrayView<int32> GetRenderDeformerNumInfluences() const;
 		/** Return a view of all the render materials used on this collection across all patterns. */
-		UE_API TConstArrayView<FString> GetRenderMaterialPathName() const;
+		UE_API TConstArrayView<FSoftObjectPath> GetRenderMaterialSoftObjectPathName() const;
 		/** Convenience to find which render pattern a render vertex belongs to */
 		UE_API int32 FindRenderPatternByVertex(int32 VertexIndex) const;
 		/** Convenience to find which render pattern a render face belongs to */
 		UE_API int32 FindRenderPatternByFaceIndex(int32 FaceIndex) const;
+		UE_DEPRECATED(5.7, "Use GetRenderMaterialSoftObjectPathName instead")
+		TConstArrayView<FString> GetRenderMaterialPathName() const
+		{
+			return TConstArrayView<FString>();
+		}
 
 		//~ Seam Group
 		/** Return the number of seams in this collection. */
@@ -250,6 +266,9 @@ namespace UE::Chaos::ClothAsset
 		/** Create this facade's groups and attributes. */
 		UE_API void DefineSchema(EClothCollectionExtendedSchemas OptionalSchemas = EClothCollectionExtendedSchemas::None);
 
+		/** Post-Serialize upgrade cloth collection */
+		UE_API void PostSerialize(const FArchive& Ar);
+
 		/** Remove all LODs from this cloth. */
 		UE_API void Reset();
 
@@ -264,9 +283,13 @@ namespace UE::Chaos::ClothAsset
 
 		//~ LOD (single per collection) Group
 		/** Set the physics asset path name. */
-		UE_API void SetPhysicsAssetPathName(const FString& PathName);
+		UE_API void SetPhysicsAssetSoftObjectPathName(const FSoftObjectPath& PathName);
 		/** Set the skeletal mesh asset path name and the reference skeleton that will be used with this asset. */
-		UE_API void SetSkeletalMeshPathName(const FString& PathName);
+		UE_API void SetSkeletalMeshSoftObjectPathName(const FSoftObjectPath& PathName);
+		UE_DEPRECATED(5.7, "Use SetPhysicsAssetSoftObjectPathName instead")
+		void SetPhysicsAssetPathName(const FString& PathName){}
+		UE_DEPRECATED(5.7, "Use SetSkeletalMeshSoftObjectPathName instead")
+		void SetSkeletalMeshPathName(const FString& PathName){}
 
 		//~ Solver (max 1 per collection) Group
 		/** Set the solver gravity */
@@ -356,7 +379,12 @@ namespace UE::Chaos::ClothAsset
 		/** Return a view of all the render deformer number of influences used on this collection across all patterns. */
 		UE_API TArrayView<int32> GetRenderDeformerNumInfluences();
 		/** Return a view of all the render materials used on this collection across all patterns. */
-		UE_API TArrayView<FString> GetRenderMaterialPathName();
+		UE_API TArrayView<FSoftObjectPath> GetRenderMaterialSoftObjectPathName();
+		UE_DEPRECATED(5.7, "Use GetRenderMaterialSoftObjectPathName instead")
+		TArrayView<FString> GetRenderMaterialPathName() 
+		{ 
+			return TArrayView<FString>(); 
+		}
 
 		//~ Seam Group
 		/** Set the new number of seams to this cloth. */
diff --git a/Engine/Plugins/ChaosClothAsset/Source/ChaosClothAsset/Public/ChaosClothAsset/CollectionClothRenderPatternFacade.h b/Engine/Plugins/ChaosClothAsset/Source/ChaosClothAsset/Public/ChaosClothAsset/CollectionClothRenderPatternFacade.h
--- a/Engine/Plugins/ChaosClothAsset/Source/ChaosClothAsset/Public/ChaosClothAsset/CollectionClothRenderPatternFacade.h
+++ b/Engine/Plugins/ChaosClothAsset/Source/ChaosClothAsset/Public/ChaosClothAsset/CollectionClothRenderPatternFacade.h
@@ -31,7 +31,13 @@ namespace UE::Chaos::ClothAsset
 		/** Return the render deformer number of influences for this pattern. */
 		UE_API int32 GetRenderDeformerNumInfluences() const;
 		/** Return the render material for this pattern. */
-		UE_API const FString& GetRenderMaterialPathName() const;
+		UE_API const FSoftObjectPath& GetRenderMaterialSoftObjectPathName() const;
+		UE_DEPRECATED(5.7, "Use GetRenderMaterialSoftObjectPathName instead")
+		const FString& GetRenderMaterialPathName() const
+		{
+			static const FString EmptyString;
+			return EmptyString;
+		}
 
 		//~ Render Vertices Group
 		// Note: Use the FCollectionClothConstFacade accessors instead of these for the array indices to match the RenderIndices values
@@ -110,7 +116,9 @@ namespace UE::Chaos::ClothAsset
 		/** Set the render deformer number of influences for this pattern. */
 		UE_API void SetRenderDeformerNumInfluences(int32 NumInfluences);
 		/** Set the render material for this pattern. */
-		UE_API void SetRenderMaterialPathName(const FString& PathName);
+		UE_API void SetRenderMaterialSoftObjectPathName(const FSoftObjectPath& PathName);
+		UE_DEPRECATED(5.7, "Use SetRenderMaterialSoftObjectPathName instead")
+		void SetRenderMaterialPathName(const FString& PathName){}
 
 		//~ Render Vertices Group
 		// Note: Use the FCollectionClothFacade accessors instead of these for the array indices to match the RenderIndices values
diff --git a/Engine/Plugins/ChaosClothAsset/Source/ChaosClothAssetEngine/Private/ChaosClothAsset/ClothAsset.cpp b/Engine/Plugins/ChaosClothAsset/Source/ChaosClothAssetEngine/Private/ChaosClothAsset/ClothAsset.cpp
--- a/Engine/Plugins/ChaosClothAsset/Source/ChaosClothAssetEngine/Private/ChaosClothAsset/ClothAsset.cpp
+++ b/Engine/Plugins/ChaosClothAsset/Source/ChaosClothAssetEngine/Private/ChaosClothAsset/ClothAsset.cpp
@@ -22,6 +22,7 @@
 #include "Rendering/SkeletalMeshModel.h"
 #include "Rendering/SkeletalMeshRenderData.h"
 #include "UObject/Package.h"
+#include "UObject/SoftObjectPath.h"
 #include "UObject/FortniteMainBranchObjectVersion.h"
 #include "EngineUtils.h"
 #if WITH_EDITOR
@@ -128,6 +129,10 @@ namespace UE::Chaos::ClothAsset::Private
 					::Chaos::Softs::FCollectionPropertyMutableFacade PropertyFacade(ClothCollection);
 					PropertyFacade.PostSerialize(Ar);
 
+					// Cloth Facade may need to upgrade
+					FCollectionClothFacade ClothFacade(ClothCollection);
+					ClothFacade.PostSerialize(Ar);
+
 					ClothCollections.Emplace(MoveTemp(ClothCollection));
 				}
 			}
@@ -519,7 +524,7 @@ void UChaosClothAsset::Build(
 	OutMaterials.Reset();
 
 	// Iterate through the LODs
-	FString PhysicsAssetPathName;
+	FSoftObjectPath PhysicsAssetPathName;
 	for (int32 LodIndex = 0; LodIndex < InClothCollections.Num(); ++LodIndex)
 	{
 		// New LOD
@@ -556,12 +561,12 @@ void UChaosClothAsset::Build(
 
 		OutMaterials.Reserve(OutMaterials.Num() + NumLodMaterials);
 
-		const TConstArrayView<FString> LodRenderMaterialPathName = ClothFacade.GetRenderMaterialPathName();
+		const TConstArrayView<FSoftObjectPath> LodRenderMaterialPathName = ClothFacade.GetRenderMaterialSoftObjectPathName();
 		for (int32 LodMaterialIndex = 0; LodMaterialIndex < NumLodMaterials; ++LodMaterialIndex)
 		{
-			const FString& RenderMaterialPathName = LodRenderMaterialPathName[LodMaterialIndex];
+			const FSoftObjectPath& RenderMaterialPathName = LodRenderMaterialPathName[LodMaterialIndex];
 
-			if (UMaterialInterface* const Material = LoadObject<UMaterialInterface>(this, *RenderMaterialPathName, nullptr, LOAD_None, nullptr))
+			if (UMaterialInterface* const Material = Cast<UMaterialInterface>(RenderMaterialPathName.TryLoad()))
 			{
 				OutMaterials.Emplace(Material, true, false, Material->GetFName());
 			}
@@ -608,10 +613,9 @@ void UChaosClothAsset::Build(
 		if (LodIndex == 0)
 		{
 			using namespace ::Chaos::Softs;
-			PhysicsAssetPathName = InClothFacade.GetPhysicsAssetPathName();
-			const FString SkeletalMeshPathName = InClothFacade.GetSkeletalMeshPathName(); 
-			USkeletalMesh* const SkeletalMesh = SkeletalMeshPathName.IsEmpty() ? nullptr :
-				LoadObject<USkeletalMesh>(nullptr, *SkeletalMeshPathName, nullptr, LOAD_None, nullptr);
+			PhysicsAssetPathName = InClothFacade.GetPhysicsAssetSoftObjectPathName();
+			const FSoftObjectPath& SkeletalMeshPathName = InClothFacade.GetSkeletalMeshSoftObjectPathName(); 
+			USkeletalMesh* const SkeletalMesh = Cast<USkeletalMesh>(SkeletalMeshPathName.TryLoad());				
 
 			// Set reference skeleton
 			SetSkeleton(SkeletalMesh ? SkeletalMesh->GetSkeleton() : nullptr); // For completion only, this is not being used and might mismatch the skeletal mesh's reference skeleton
@@ -639,9 +643,7 @@ void UChaosClothAsset::Build(
 	}
 
 	// Set physics asset (note: the cloth asset's physics asset is only replaced if a collection path name is found valid)
-	PhysicsAsset = !PhysicsAssetPathName.IsEmpty() ?
-		LoadObject<UPhysicsAsset>(this, *PhysicsAssetPathName, nullptr, LOAD_None, nullptr) :
-		nullptr;
+	PhysicsAsset = Cast<UPhysicsAsset>(PhysicsAssetPathName.TryLoad());
 
 	SetHasVertexColors(true);
 
@@ -879,9 +881,8 @@ void UChaosClothAsset::UpdateSkeletonFromCollection(bool /*bRebuildModels*/)
 	FCollectionClothConstFacade ClothFacade(GetClothCollections()[0]);
 	check(ClothFacade.IsValid());
 
-	const FString& SkeletalMeshPathName = ClothFacade.GetSkeletalMeshPathName();
-	USkeletalMesh* const SkeletalMesh = SkeletalMeshPathName.IsEmpty() ? nullptr :
-		LoadObject<USkeletalMesh>(nullptr, *SkeletalMeshPathName, nullptr, LOAD_None, nullptr);
+	const FSoftObjectPath& SkeletalMeshPathName = ClothFacade.GetSkeletalMeshSoftObjectPathName();
+	USkeletalMesh* const SkeletalMesh = Cast<USkeletalMesh>(SkeletalMeshPathName.TryLoad());
 
 	SetSkeleton(SkeletalMesh ? SkeletalMesh->GetSkeleton() : nullptr); // For completion only, this is not being used and might mismatch the skeletal mesh's reference skeleton
 	Super::SetReferenceSkeleton(SkeletalMesh ? &SkeletalMesh->GetRefSkeleton() : nullptr);
@@ -893,9 +894,9 @@ void UChaosClothAsset::CopySimMeshToRenderMesh(UMaterialInterface* Material)
 	check(GetClothCollections().Num());
 
 	// Add a default material if none is specified
-	const FString RenderMaterialPathName = Material ?
+	const FSoftObjectPath RenderMaterialPathName = FSoftObjectPath(Material ?
 		Material->GetPathName() :
-		FString(TEXT("/Engine/EditorMaterials/Cloth/CameraLitDoubleSided.CameraLitDoubleSided"));
+		FString(TEXT("/Engine/EditorMaterials/Cloth/CameraLitDoubleSided.CameraLitDoubleSided")));
 
 	bool bAnyLodHasRenderMesh = false;
 	for (TSharedRef<const FManagedArrayCollection>& ClothCollection : GetClothCollections())
@@ -911,7 +912,7 @@ void UChaosClothAsset::CopySimMeshToRenderMesh(UMaterialInterface* Material)
 	Materials.Reset(1);
 	if (bAnyLodHasRenderMesh)
 	{
-		if (UMaterialInterface* const LoadedMaterial = LoadObject<UMaterialInterface>(nullptr, *RenderMaterialPathName, nullptr, LOAD_None, nullptr))
+		if (UMaterialInterface* const LoadedMaterial = Cast<UMaterialInterface>(RenderMaterialPathName.TryLoad()))
 		{
 			Materials.Emplace(LoadedMaterial, true, false, LoadedMaterial->GetFName());
 		}
diff --git a/Engine/Plugins/ChaosClothAsset/Source/ChaosClothAssetEngine/Private/ChaosClothAsset/ClothAssetBuilder.cpp b/Engine/Plugins/ChaosClothAsset/Source/ChaosClothAssetEngine/Private/ChaosClothAsset/ClothAssetBuilder.cpp
--- a/Engine/Plugins/ChaosClothAsset/Source/ChaosClothAssetEngine/Private/ChaosClothAsset/ClothAssetBuilder.cpp
+++ b/Engine/Plugins/ChaosClothAsset/Source/ChaosClothAssetEngine/Private/ChaosClothAsset/ClothAssetBuilder.cpp
@@ -219,7 +219,7 @@ void UChaosClothAsset::FBuilder::BuildLod(FSkeletalMeshLODModel& LODModel, const
 			const int32 MaterialIndex = MaterialOffset + SectionIndex;
 			if (!ClothAsset.GetMaterials().IsValidIndex(MaterialIndex))
 			{
-				const FString& MaterialPath = RenderPatternFacade.GetRenderMaterialPathName();
+				const FString MaterialPath = RenderPatternFacade.GetRenderMaterialSoftObjectPathName().ToString();
 				UE_LOG(LogChaosClothAsset, Warning, TEXT("Cloth Asset BuildLod: Pattern %d in LOD %d has invalid material index. Material index is %d, asset has %d materials. Expected material path for this pattern: %s"),
 					SectionIndex, LodIndex, MaterialIndex, ClothAsset.GetMaterials().Num(), *MaterialPath);
 			}
diff --git a/Engine/Plugins/ChaosClothAssetEditor/Source/ChaosClothAssetDataflowNodes/Private/ChaosClothAsset/ChaosClothAssetDataflowNodesModule.cpp b/Engine/Plugins/ChaosClothAssetEditor/Source/ChaosClothAssetDataflowNodes/Private/ChaosClothAsset/ChaosClothAssetDataflowNodesModule.cpp
--- a/Engine/Plugins/ChaosClothAssetEditor/Source/ChaosClothAssetDataflowNodes/Private/ChaosClothAsset/ChaosClothAssetDataflowNodesModule.cpp
+++ b/Engine/Plugins/ChaosClothAssetEditor/Source/ChaosClothAssetDataflowNodes/Private/ChaosClothAsset/ChaosClothAssetDataflowNodesModule.cpp
@@ -310,7 +310,12 @@ PRAGMA_ENABLE_DEPRECATION_WARNINGS
 
 					if (State.GetViewMode().GetName() == FClothRenderViewMode::Name)
 					{
-						TArray<FString> MaterialArray(ClothFacade.GetRenderMaterialPathName());
+						TArray<FString> MaterialArray;
+						MaterialArray.Reserve(ClothFacade.GetRenderMaterialSoftObjectPathName().Num());
+						for (const FSoftObjectPath& MaterialName : ClothFacade.GetRenderMaterialSoftObjectPathName())
+						{
+							MaterialArray.Emplace(MaterialName.ToString());
+						}
 						RenderCollection.AddSurface(MoveTemp(Vertices), MoveTemp(Indices), MoveTemp(Normals), MoveTemp(Colors), MoveTemp(UVs), MoveTemp(RenderMaterialIDs), MoveTemp(MaterialArray));
 					}
 					else
diff --git a/Engine/Plugins/ChaosClothAssetEditor/Source/ChaosClothAssetDataflowNodes/Private/ChaosClothAsset/ClothCollectionToDynamicMeshNode.cpp b/Engine/Plugins/ChaosClothAssetEditor/Source/ChaosClothAssetDataflowNodes/Private/ChaosClothAsset/ClothCollectionToDynamicMeshNode.cpp
--- a/Engine/Plugins/ChaosClothAssetEditor/Source/ChaosClothAssetDataflowNodes/Private/ChaosClothAsset/ClothCollectionToDynamicMeshNode.cpp
+++ b/Engine/Plugins/ChaosClothAssetEditor/Source/ChaosClothAssetDataflowNodes/Private/ChaosClothAsset/ClothCollectionToDynamicMeshNode.cpp
@@ -62,12 +62,12 @@ void FChaosClothAssetCollectionToDynamicMeshNode::Evaluate(UE::Dataflow::FContex
 			SetValue(Context, NewMesh, &RenderDynamicMesh);
 
 			TArray<TObjectPtr<UMaterialInterface>> OutMaterials;
-			const TConstArrayView<FString> RenderMaterialPathNames = ClothFacade.GetRenderMaterialPathName();
+			const TConstArrayView<FSoftObjectPath> RenderMaterialPathNames = ClothFacade.GetRenderMaterialSoftObjectPathName();
 			OutMaterials.Reserve(RenderMaterialPathNames.Num());
 			for (int32 RenderPatternIndex = 0; RenderPatternIndex < RenderMaterialPathNames.Num(); ++RenderPatternIndex)
 			{
-				const FString& RenderMaterialPathName = RenderMaterialPathNames[RenderPatternIndex];
-				OutMaterials.Add(LoadObject<UMaterialInterface>((UObject*)GetTransientPackage(), *RenderMaterialPathName, nullptr, LOAD_None, nullptr));
+				const FSoftObjectPath& RenderMaterialPathName = RenderMaterialPathNames[RenderPatternIndex];
+				OutMaterials.Add(Cast<UMaterialInterface>(RenderMaterialPathName.TryLoad()));
 			}
 			SetValue(Context, OutMaterials, &RenderMaterials);
 			return;
@@ -101,7 +101,7 @@ void FChaosClothAssetUpdateClothFromDynamicMeshNode::Evaluate(UE::Dataflow::FCon
 			if (bCopyToRenderMaterials)
 			{
 				const TArray<TObjectPtr<const UMaterialInterface>> InMaterials = GetValue(Context, &Materials);
-				TArrayView<FString> RenderMaterialPathNames = ClothFacade.GetRenderMaterialPathName();
+				TArrayView<FSoftObjectPath> RenderMaterialPathNames = ClothFacade.GetRenderMaterialSoftObjectPathName();
 				for (int32 MaterialIndex = 0; MaterialIndex < FMath::Min(RenderMaterialPathNames.Num(), InMaterials.Num()); ++MaterialIndex)
 				{
 					RenderMaterialPathNames[MaterialIndex] = InMaterials[MaterialIndex]->GetPathName();
diff --git a/Engine/Plugins/ChaosClothAssetEditor/Source/ChaosClothAssetDataflowNodes/Private/ChaosClothAsset/ClothDataflowTools.cpp b/Engine/Plugins/ChaosClothAssetEditor/Source/ChaosClothAssetDataflowNodes/Private/ChaosClothAsset/ClothDataflowTools.cpp
--- a/Engine/Plugins/ChaosClothAssetEditor/Source/ChaosClothAssetDataflowNodes/Private/ChaosClothAsset/ClothDataflowTools.cpp
+++ b/Engine/Plugins/ChaosClothAssetEditor/Source/ChaosClothAssetDataflowNodes/Private/ChaosClothAsset/ClothDataflowTools.cpp
@@ -439,7 +439,7 @@ namespace UE::Chaos::ClothAsset
 		} // namespace Private::Resizing
 	} // Private
 
-	void FClothDataflowTools::AddRenderPatternFromSkeletalMeshSection(const TSharedRef<FManagedArrayCollection>& ClothCollection, const FSkeletalMeshLODModel& SkeletalMeshModel, const int32 SectionIndex, const FString& RenderMaterialPathName)
+	void FClothDataflowTools::AddRenderPatternFromSkeletalMeshSection(const TSharedRef<FManagedArrayCollection>& ClothCollection, const FSkeletalMeshLODModel& SkeletalMeshModel, const int32 SectionIndex, const FSoftObjectPath& RenderMaterialPathName)
 	{
 		check(SectionIndex < SkeletalMeshModel.Sections.Num());
 
@@ -498,7 +498,7 @@ namespace UE::Chaos::ClothAsset
 				SkeletalMeshModel.IndexBuffer[IndexOffset + 2] - Section.BaseVertexIndex + VertexOffset
 			);
 		}
-		ClothPatternFacade.SetRenderMaterialPathName(RenderMaterialPathName);
+		ClothPatternFacade.SetRenderMaterialSoftObjectPathName(RenderMaterialPathName);
 	}
 	
 	void FClothDataflowTools::AddSimPatternsFromSkeletalMeshSection(const TSharedRef<FManagedArrayCollection>& ClothCollection, const FSkeletalMeshLODModel& SkeletalMeshModel, const int32 SectionIndex, const int32 UVChannelIndex, const FVector2f& UVScale, bool bImportNormals, TArray<int32>* OutSim2DToSourceVertex)
diff --git a/Engine/Plugins/ChaosClothAssetEditor/Source/ChaosClothAssetDataflowNodes/Private/ChaosClothAsset/ClothDataflowTools.h b/Engine/Plugins/ChaosClothAssetEditor/Source/ChaosClothAssetDataflowNodes/Private/ChaosClothAsset/ClothDataflowTools.h
--- a/Engine/Plugins/ChaosClothAssetEditor/Source/ChaosClothAssetDataflowNodes/Private/ChaosClothAsset/ClothDataflowTools.h
+++ b/Engine/Plugins/ChaosClothAssetEditor/Source/ChaosClothAssetDataflowNodes/Private/ChaosClothAsset/ClothDataflowTools.h
@@ -23,7 +23,7 @@ namespace UE::Chaos::ClothAsset
 	*/
 	struct FClothDataflowTools
 	{
-		static void AddRenderPatternFromSkeletalMeshSection(const TSharedRef<FManagedArrayCollection>& ClothCollection, const FSkeletalMeshLODModel& SkeletalMeshModel, const int32 SectionIndex, const FString& RenderMaterialPathName);
+		static void AddRenderPatternFromSkeletalMeshSection(const TSharedRef<FManagedArrayCollection>& ClothCollection, const FSkeletalMeshLODModel& SkeletalMeshModel, const int32 SectionIndex, const FSoftObjectPath& RenderMaterialPathName);
 
 		static void AddSimPatternsFromSkeletalMeshSection(const TSharedRef<FManagedArrayCollection>& ClothCollection, const FSkeletalMeshLODModel& SkeletalMeshModel, const int32 SectionIndex, const int32 UVChannelIndex, const FVector2f& UVScale, bool bImportNormals = false, TArray<int32>* OutSim2DToSourceVertex = nullptr);
 
diff --git a/Engine/Plugins/ChaosClothAssetEditor/Source/ChaosClothAssetDataflowNodes/Private/ChaosClothAsset/MergeClothCollectionsNode.cpp b/Engine/Plugins/ChaosClothAssetEditor/Source/ChaosClothAssetDataflowNodes/Private/ChaosClothAsset/MergeClothCollectionsNode.cpp
--- a/Engine/Plugins/ChaosClothAssetEditor/Source/ChaosClothAssetDataflowNodes/Private/ChaosClothAsset/MergeClothCollectionsNode.cpp
+++ b/Engine/Plugins/ChaosClothAssetEditor/Source/ChaosClothAssetDataflowNodes/Private/ChaosClothAsset/MergeClothCollectionsNode.cpp
@@ -8,6 +8,7 @@
 #include "Chaos/CollectionPropertyFacade.h"
 #include "Dataflow/DataflowInputOutput.h"
 #include "Engine/SkeletalMesh.h"
+#include "UObject/SoftObjectPath.h"
 #include "Chaos/CollectionEmbeddedSpringConstraintFacade.h"
 
 #include UE_INLINE_GENERATED_CPP_BY_NAME(MergeClothCollectionsNode)
@@ -448,23 +449,23 @@ namespace UE::Chaos::ClothAsset::Private
 		OtherBoneRemap.Reset();
 
 		/** Disallow merging cloth facades with incompatible ref skeletons. */
-		const FString& SkeletalMeshPathName1 = Cloth1.GetSkeletalMeshPathName();
-		const FString& SkeletalMeshPathName2 = Cloth2.GetSkeletalMeshPathName();
-		if (SkeletalMeshPathName1.IsEmpty() || SkeletalMeshPathName2.IsEmpty() || SkeletalMeshPathName1 == SkeletalMeshPathName2)
+		const FSoftObjectPath& SkeletalMeshPathName1 = Cloth1.GetSkeletalMeshSoftObjectPathName();
+		const FSoftObjectPath& SkeletalMeshPathName2 = Cloth2.GetSkeletalMeshSoftObjectPathName();
+		if (SkeletalMeshPathName1.IsNull() || SkeletalMeshPathName2.IsNull() || SkeletalMeshPathName1 == SkeletalMeshPathName2)
 		{
 			return true;
 		}
 
 		static const FText ErrorHeadline = LOCTEXT("IncompatibleSkeletalMeshesHeadline", "Incompatible Skeletal Meshes.");
-		const USkeletalMesh* const SkeletalMesh1 = LoadObject<USkeletalMesh>(nullptr, *SkeletalMeshPathName1, nullptr, LOAD_None, nullptr);
-		const USkeletalMesh* const SkeletalMesh2 = LoadObject<USkeletalMesh>(nullptr, *SkeletalMeshPathName2, nullptr, LOAD_None, nullptr);
+		const USkeletalMesh* const SkeletalMesh1 = Cast<USkeletalMesh>(SkeletalMeshPathName1.TryLoad());
+		const USkeletalMesh* const SkeletalMesh2 = Cast<USkeletalMesh>(SkeletalMeshPathName2.TryLoad());
 		if (!SkeletalMesh1 || !SkeletalMesh2)
 		{
 			const FText Details = FText::Format(
 				LOCTEXT(
 					"IncompatibleSkeletalMeshesLoadFailureDetails",
 					"Cloth collections failed to merge due to failing to load SkeletalMesh \"{0}\" to check compatibility."),
-				!SkeletalMesh1 ? FText::FromString(SkeletalMeshPathName1) : FText::FromString(SkeletalMeshPathName2));
+				!SkeletalMesh1 ? FText::FromString(SkeletalMeshPathName1.ToString()) : FText::FromString(SkeletalMeshPathName2.ToString()));
 
 			FClothDataflowTools::LogAndToastWarning(DataflowNode, ErrorHeadline, Details);
 			return false;
@@ -475,7 +476,7 @@ namespace UE::Chaos::ClothAsset::Private
 
 		const FReferenceSkeleton& MergedRefSkeleton = RefSkeleton1.GetNum() >= RefSkeleton2.GetNum() ? RefSkeleton1 : RefSkeleton2;
 		const FReferenceSkeleton& RemapRefSkeleton = RefSkeleton1.GetNum() >= RefSkeleton2.GetNum() ? RefSkeleton2 : RefSkeleton1;
-		const FString& MergedSkeletalMeshPath = RefSkeleton1.GetNum() >= RefSkeleton2.GetNum() ? SkeletalMeshPathName1 : SkeletalMeshPathName2;
+		const FSoftObjectPath& MergedSkeletalMeshPath = RefSkeleton1.GetNum() >= RefSkeleton2.GetNum() ? SkeletalMeshPathName1 : SkeletalMeshPathName2;
 
 		const TArray<FMeshBoneInfo>& RemapBoneInfo = RemapRefSkeleton.GetRefBoneInfo();
 		const TArray<FTransform>& MergedBonePose = MergedRefSkeleton.GetRefBonePose();
@@ -492,10 +493,10 @@ namespace UE::Chaos::ClothAsset::Private
 					LOCTEXT(
 						"IncompatibleSkeletalMeshesRefBoneInfoDetails",
 						"Cloth collections failed to merge due to incompatible Skeletal Meshes, \"{0}\" and \"{1}\". Could not find bone \"{2}\" in \"{3}\"."),
-					FText::FromString(SkeletalMeshPathName1),
-					FText::FromString(SkeletalMeshPathName2),
+					FText::FromString(SkeletalMeshPathName1.ToString()),
+					FText::FromString(SkeletalMeshPathName2.ToString()),
 					FText::FromName(RemapBoneInfo[BoneIndex].Name),
-					FText::FromString(MergedSkeletalMeshPath));
+					FText::FromString(MergedSkeletalMeshPath.ToString()));
 
 				FClothDataflowTools::LogAndToastWarning(DataflowNode, ErrorHeadline, Details);
 				return false;
@@ -506,8 +507,8 @@ namespace UE::Chaos::ClothAsset::Private
 					LOCTEXT(
 						"IncompatibleSkeletalMeshesRefBonePoseDetails",
 						"Cloth collections failed to merge due to incompatible Skeletal Meshes, \"{0}\" and \"{1}\". RefBonePoses are mismatched for bone \"{2}\"."),
-					FText::FromString(SkeletalMeshPathName1),
-					FText::FromString(SkeletalMeshPathName2),
+					FText::FromString(SkeletalMeshPathName1.ToString()),
+					FText::FromString(SkeletalMeshPathName2.ToString()),
 					FText::FromName(RemapBoneInfo[BoneIndex].Name));
 
 				FClothDataflowTools::LogAndToastWarning(DataflowNode, ErrorHeadline, Details);
@@ -532,7 +533,7 @@ namespace UE::Chaos::ClothAsset::Private
 				OtherBoneRemap = MoveTemp(RemapIndices);
 			}
 		}
-		Cloth1.SetSkeletalMeshPathName(MergedSkeletalMeshPath);
+		Cloth1.SetSkeletalMeshSoftObjectPathName(MergedSkeletalMeshPath);
 
 		return true;
 	}
diff --git a/Engine/Plugins/ChaosClothAssetEditor/Source/ChaosClothAssetDataflowNodes/Private/ChaosClothAsset/RemeshNode.cpp b/Engine/Plugins/ChaosClothAssetEditor/Source/ChaosClothAssetDataflowNodes/Private/ChaosClothAsset/RemeshNode.cpp
--- a/Engine/Plugins/ChaosClothAssetEditor/Source/ChaosClothAssetDataflowNodes/Private/ChaosClothAsset/RemeshNode.cpp
+++ b/Engine/Plugins/ChaosClothAssetEditor/Source/ChaosClothAssetDataflowNodes/Private/ChaosClothAsset/RemeshNode.cpp
@@ -1772,7 +1772,7 @@ namespace UE::Chaos::ClothAsset::Private
 
 			FCollectionClothConstFacade InClothFacade(ClothCollection);
 			FCollectionClothRenderPatternConstFacade InPatternFacade = InClothFacade.GetRenderPattern(SourceMaterialID);
-			OutClothPatternFacade.SetRenderMaterialPathName(InPatternFacade.GetRenderMaterialPathName());
+			OutClothPatternFacade.SetRenderMaterialSoftObjectPathName(InPatternFacade.GetRenderMaterialSoftObjectPathName());
 		}
 
 
diff --git a/Engine/Plugins/ChaosClothAssetEditor/Source/ChaosClothAssetDataflowNodes/Private/ChaosClothAsset/RenderMeshImport.cpp b/Engine/Plugins/ChaosClothAssetEditor/Source/ChaosClothAssetDataflowNodes/Private/ChaosClothAsset/RenderMeshImport.cpp
--- a/Engine/Plugins/ChaosClothAssetEditor/Source/ChaosClothAssetDataflowNodes/Private/ChaosClothAsset/RenderMeshImport.cpp
+++ b/Engine/Plugins/ChaosClothAssetEditor/Source/ChaosClothAssetDataflowNodes/Private/ChaosClothAsset/RenderMeshImport.cpp
@@ -358,7 +358,7 @@ namespace UE::Chaos::ClothAsset
 				const FString RenderMaterialPathName = Materials[MaterialIndex].MaterialInterface ?
 					Materials[MaterialIndex].MaterialInterface->GetPathName() :
 					FString();
-				ClothPatternFacade.SetRenderMaterialPathName(RenderMaterialPathName);
+				ClothPatternFacade.SetRenderMaterialSoftObjectPathName(RenderMaterialPathName);
 			}
 		}
 	}
diff --git a/Engine/Plugins/ChaosClothAssetEditor/Source/ChaosClothAssetDataflowNodes/Private/ChaosClothAsset/SetPhysicsAssetNode.cpp b/Engine/Plugins/ChaosClothAssetEditor/Source/ChaosClothAssetDataflowNodes/Private/ChaosClothAsset/SetPhysicsAssetNode.cpp
--- a/Engine/Plugins/ChaosClothAssetEditor/Source/ChaosClothAssetDataflowNodes/Private/ChaosClothAsset/SetPhysicsAssetNode.cpp
+++ b/Engine/Plugins/ChaosClothAssetEditor/Source/ChaosClothAssetDataflowNodes/Private/ChaosClothAsset/SetPhysicsAssetNode.cpp
@@ -32,7 +32,7 @@ void FChaosClothAssetSetPhysicsAssetNode::Evaluate(UE::Dataflow::FContext& Conte
 		{
 			FCollectionClothFacade CollectionClothFacade(ClothCollection);
 			TObjectPtr<const UPhysicsAsset> InPhysicsAsset = GetValue(Context, &PhysicsAsset);
-			CollectionClothFacade.SetPhysicsAssetPathName(InPhysicsAsset ? InPhysicsAsset.GetPathName() : FString());
+			CollectionClothFacade.SetPhysicsAssetSoftObjectPathName(InPhysicsAsset ? InPhysicsAsset.GetPathName() : FString());
 		}
 
 		SetValue(Context, MoveTemp(*ClothCollection), &Collection);
diff --git a/Engine/Plugins/ChaosClothAssetEditor/Source/ChaosClothAssetDataflowNodes/Private/ChaosClothAsset/SkeletalMeshImportNode.cpp b/Engine/Plugins/ChaosClothAssetEditor/Source/ChaosClothAssetDataflowNodes/Private/ChaosClothAsset/SkeletalMeshImportNode.cpp
--- a/Engine/Plugins/ChaosClothAssetEditor/Source/ChaosClothAssetDataflowNodes/Private/ChaosClothAsset/SkeletalMeshImportNode.cpp
+++ b/Engine/Plugins/ChaosClothAssetEditor/Source/ChaosClothAssetDataflowNodes/Private/ChaosClothAsset/SkeletalMeshImportNode.cpp
@@ -170,10 +170,10 @@ void FChaosClothAssetSkeletalMeshImportNode_v2::Evaluate(UE::Dataflow::FContext&
 
 			if (const UPhysicsAsset* PhysicsAsset = bSetPhysicsAsset ? InSkeletalMesh->GetPhysicsAsset() : nullptr)
 			{
-				ClothFacade.SetPhysicsAssetPathName(PhysicsAsset->GetPathName());
+				ClothFacade.SetPhysicsAssetSoftObjectPathName(PhysicsAsset->GetPathName());
 			}
 
-			ClothFacade.SetSkeletalMeshPathName(InSkeletalMesh->GetPathName());
+			ClothFacade.SetSkeletalMeshSoftObjectPathName(InSkeletalMesh->GetPathName());
 		}
 		SetValue(Context, MoveTemp(*ClothCollection), &Collection);
 	}
@@ -249,7 +249,7 @@ void FChaosClothAssetSkeletalMeshImportNode::Evaluate(UE::Dataflow::FContext& Co
 
 			if (const UPhysicsAsset* PhysicsAsset = bSetPhysicsAsset ? InSkeletalMesh->GetPhysicsAsset() : nullptr)
 			{
-				ClothFacade.SetPhysicsAssetPathName(PhysicsAsset->GetPathName());
+				ClothFacade.SetPhysicsAssetSoftObjectPathName(PhysicsAsset->GetPathName());
 			}
 
 			// In order to retain existing behavior, flip the sim normals.
@@ -257,7 +257,7 @@ void FChaosClothAssetSkeletalMeshImportNode::Evaluate(UE::Dataflow::FContext& Co
 			constexpr bool bReverseFalse = false;
 			FClothGeometryTools::ReverseMesh(ClothCollection, bReverseSimMeshNormals, bReverseFalse, bReverseFalse, bReverseFalse, TArray<int32>(), TArray<int32>());
 
-			ClothFacade.SetSkeletalMeshPathName(InSkeletalMesh->GetPathName());
+			ClothFacade.SetSkeletalMeshSoftObjectPathName(InSkeletalMesh->GetPathName());
 		}
 		SetValue(Context, MoveTemp(*ClothCollection), &Collection);
 	}
diff --git a/Engine/Plugins/ChaosClothAssetEditor/Source/ChaosClothAssetDataflowNodes/Private/ChaosClothAsset/TransferSkinWeightsNode.cpp b/Engine/Plugins/ChaosClothAssetEditor/Source/ChaosClothAssetDataflowNodes/Private/ChaosClothAsset/TransferSkinWeightsNode.cpp
--- a/Engine/Plugins/ChaosClothAssetEditor/Source/ChaosClothAssetDataflowNodes/Private/ChaosClothAsset/TransferSkinWeightsNode.cpp
+++ b/Engine/Plugins/ChaosClothAssetEditor/Source/ChaosClothAssetDataflowNodes/Private/ChaosClothAsset/TransferSkinWeightsNode.cpp
@@ -902,15 +902,15 @@ void FChaosClothAssetTransferSkinWeightsNode::Evaluate(UE::Dataflow::FContext& C
 				MeshTransforms::ApplyTransform(SourceDynamicMesh, Transform, true);
 				TargetRefSkeleton = &InSkeletalMesh->GetRefSkeleton();
 
-				ClothFacade.SetSkeletalMeshPathName(InSkeletalMesh->GetPathName());
+				ClothFacade.SetSkeletalMeshSoftObjectPathName(InSkeletalMesh->GetPathName());
 			}
 			else
 			{
 				// Reuse the input sim mesh skeleton
 				auto GetCollectionRefSkeleton = [&InSkeletalMesh](const FCollectionClothConstFacade& ClothFacade) -> const FReferenceSkeleton*
 					{
-						const FString& SkeletalMeshPathName = ClothFacade.GetSkeletalMeshPathName();
-						InSkeletalMesh = TStrongObjectPtr<const USkeletalMesh>(LoadObject<USkeletalMesh>(nullptr, *SkeletalMeshPathName));
+						const FSoftObjectPath& SkeletalMeshPathName = ClothFacade.GetSkeletalMeshSoftObjectPathName();
+						InSkeletalMesh = TStrongObjectPtr<const USkeletalMesh>(Cast<USkeletalMesh>(SkeletalMeshPathName.TryLoad()));
 						return InSkeletalMesh ? &InSkeletalMesh->GetRefSkeleton() : nullptr;
 					};
 
diff --git a/Engine/Plugins/ChaosClothAssetEditor/Source/ChaosClothAssetEditor/Private/ChaosClothAsset/ClothEditorMode.cpp b/Engine/Plugins/ChaosClothAssetEditor/Source/ChaosClothAssetEditor/Private/ChaosClothAsset/ClothEditorMode.cpp
--- a/Engine/Plugins/ChaosClothAssetEditor/Source/ChaosClothAssetEditor/Private/ChaosClothAsset/ClothEditorMode.cpp
+++ b/Engine/Plugins/ChaosClothAssetEditor/Source/ChaosClothAssetEditor/Private/ChaosClothAsset/ClothEditorMode.cpp
@@ -774,11 +774,11 @@ void UChaosClothAssetEditorMode::ReinitializeDynamicMeshComponents()
 			break;
 			case EClothPatternVertexType::Render:
 			{
-				const TArrayView<const FString> MaterialPaths = ClothFacade.GetRenderMaterialPathName();
+				const TArrayView<const FSoftObjectPath> MaterialPaths = ClothFacade.GetRenderMaterialSoftObjectPathName();
 				for (int32 MaterialIndex = 0; MaterialIndex < MaterialPaths.Num(); ++MaterialIndex)
 				{
-					const FString& Path = MaterialPaths[MaterialIndex];
-					UMaterialInterface* const Material = LoadObject<UMaterialInterface>(nullptr, *Path);
+					const FSoftObjectPath& Path = MaterialPaths[MaterialIndex];
+					UMaterialInterface* const Material = Cast<UMaterialInterface>(Path.TryLoad());
 					MeshComponent.SetMaterial(MaterialIndex, Material);
 				}
 
diff --git a/Engine/Plugins/ChaosClothAssetEditor/Source/ChaosClothAssetTools/Private/ChaosClothAsset/ClothPatternToDynamicMesh.cpp b/Engine/Plugins/ChaosClothAssetEditor/Source/ChaosClothAssetTools/Private/ChaosClothAsset/ClothPatternToDynamicMesh.cpp
--- a/Engine/Plugins/ChaosClothAssetEditor/Source/ChaosClothAssetTools/Private/ChaosClothAsset/ClothPatternToDynamicMesh.cpp
+++ b/Engine/Plugins/ChaosClothAssetEditor/Source/ChaosClothAssetTools/Private/ChaosClothAsset/ClothPatternToDynamicMesh.cpp
@@ -231,10 +231,8 @@ public:
 		}
 
 		// Set the reference skeleton if available
-		const FString& SkeletalMeshPathName = ClothFacade.GetSkeletalMeshPathName();
-		const USkeletalMesh* const SkeletalMesh = !SkeletalMeshPathName.IsEmpty() ?
-			LoadObject<USkeletalMesh>(nullptr, *SkeletalMeshPathName, nullptr, LOAD_None, nullptr) :
-			nullptr;
+		const FSoftObjectPath& SkeletalMeshPathName = ClothFacade.GetSkeletalMeshSoftObjectPathName();
+		const USkeletalMesh* const SkeletalMesh = Cast<USkeletalMesh>(SkeletalMeshPathName.TryLoad());
 		RefSkeleton = SkeletalMesh ? &SkeletalMesh->GetRefSkeleton() : nullptr;
 	}
 
@@ -713,7 +711,7 @@ void FClothPatternToDynamicMesh::Convert(
 	// Fix up any triangles without valid material IDs and apply material offset
 	if (VertexDataType == EClothPatternVertexType::Render && MaterialOffset != INDEX_NONE)
 	{
-		const TConstArrayView<FString> MaterialPaths = ClothFacade.GetRenderMaterialPathName();
+		const TConstArrayView<FSoftObjectPath> MaterialPaths = ClothFacade.GetRenderMaterialSoftObjectPathName();
 		int32 DefaultMaterialID = 0;
 		for (const int32 TriID : MeshOut.TriangleIndicesItr())
 		{
diff --git a/Engine/Plugins/ChaosClothAssetEditor/Source/ChaosClothAssetTools/Private/ChaosClothAsset/ClothingAssetToClothAssetExporter.cpp b/Engine/Plugins/ChaosClothAssetEditor/Source/ChaosClothAssetTools/Private/ChaosClothAsset/ClothingAssetToClothAssetExporter.cpp
--- a/Engine/Plugins/ChaosClothAssetEditor/Source/ChaosClothAssetTools/Private/ChaosClothAsset/ClothingAssetToClothAssetExporter.cpp
+++ b/Engine/Plugins/ChaosClothAssetEditor/Source/ChaosClothAssetTools/Private/ChaosClothAsset/ClothingAssetToClothAssetExporter.cpp
@@ -130,9 +130,9 @@ void UClothingAssetToChaosClothAssetExporter::Export(const UClothingAssetBase* C
 		FClothGeometryTools::BuildSimMeshFromDynamicMesh(ClothCollection, DynamicMesh, UVChannelIndexNone, FVector2f(1.f), bAppend);
 
 		// Set the physics asset if any
-		ClothFacade.SetPhysicsAssetPathName(PhysicsAssetPathName);
+		ClothFacade.SetPhysicsAssetSoftObjectPathName(PhysicsAssetPathName);
 		// Set the skeleton
-		ClothFacade.SetSkeletalMeshPathName(SkeletalMeshPathName);
+		ClothFacade.SetSkeletalMeshSoftObjectPathName(SkeletalMeshPathName);
 
 		// Bind the sim mesh to root bone
 		constexpr bool bBindSimMesh = true;
@@ -152,8 +152,8 @@ void UClothingAssetToChaosClothAssetExporter::Export(const UClothingAssetBase* C
 		TSharedRef<FManagedArrayCollection> ClothCollection = MakeShared<FManagedArrayCollection>();
 		FCollectionClothFacade ClothFacade(ClothCollection);
 		ClothFacade.DefineSchema();
-		ClothFacade.SetPhysicsAssetPathName(PhysicsAssetPathName);
-		ClothFacade.SetSkeletalMeshPathName(SkeletalMeshPathName);
+		ClothFacade.SetPhysicsAssetSoftObjectPathName(PhysicsAssetPathName);
+		ClothFacade.SetSkeletalMeshSoftObjectPathName(SkeletalMeshPathName);
 		ClothCollections.Emplace(MoveTemp(ClothCollection));
 	}
 
diff --git a/Engine/Plugins/ChaosClothAssetEditor/Source/ChaosClothAssetTools/Private/ChaosClothAsset/SkeletalMeshConverter.cpp b/Engine/Plugins/ChaosClothAssetEditor/Source/ChaosClothAssetTools/Private/ChaosClothAsset/SkeletalMeshConverter.cpp
--- a/Engine/Plugins/ChaosClothAssetEditor/Source/ChaosClothAssetTools/Private/ChaosClothAsset/SkeletalMeshConverter.cpp
+++ b/Engine/Plugins/ChaosClothAssetEditor/Source/ChaosClothAssetTools/Private/ChaosClothAsset/SkeletalMeshConverter.cpp
@@ -70,11 +70,11 @@ bool UClothAssetEditorSkeletalMeshConverter::ExportToSkeletalMesh(const UChaosCl
 					DynamicMeshes[LodIndex].AppendWithOffsets(DynamicMeshPiece);
 				}
 
-				const TArrayView<const FString> MaterialPaths = ClothFacade.GetRenderMaterialPathName();
+				const TArrayView<const FSoftObjectPath> MaterialPaths = ClothFacade.GetRenderMaterialSoftObjectPathName();
 				RenderMaterials.Reserve(RenderMaterials.Num() + MaterialPaths.Num());
 				for (int32 MaterialIndex = 0; MaterialIndex < MaterialPaths.Num(); ++MaterialIndex)
 				{
-					RenderMaterials.Emplace(LoadObject<UMaterialInterface>(nullptr, *MaterialPaths[MaterialIndex]));
+					RenderMaterials.Emplace(Cast<UMaterialInterface>(MaterialPaths[MaterialIndex].TryLoad()));
 				}
 				MaterialOffset += MaterialPaths.Num();
 			}
diff --git a/Engine/Plugins/Experimental/ChaosOutfitAsset/Source/ChaosOutfitAssetDataflowNodes/Private/ChaosOutfitAsset/GetOutfitBodyPartsNode.cpp b/Engine/Plugins/Experimental/ChaosOutfitAsset/Source/ChaosOutfitAssetDataflowNodes/Private/ChaosOutfitAsset/GetOutfitBodyPartsNode.cpp
--- a/Engine/Plugins/Experimental/ChaosOutfitAsset/Source/ChaosOutfitAssetDataflowNodes/Private/ChaosOutfitAsset/GetOutfitBodyPartsNode.cpp
+++ b/Engine/Plugins/Experimental/ChaosOutfitAsset/Source/ChaosOutfitAssetDataflowNodes/Private/ChaosOutfitAsset/GetOutfitBodyPartsNode.cpp
@@ -42,12 +42,12 @@ void FChaosGetOutfitBodyPartsNode::Evaluate(UE::Dataflow::FContext& Context, con
 		OutBodyParts.SetNum(NumBodySizes);
 		for (int32 BodySizeIndex = 0; BodySizeIndex < NumBodySizes; ++BodySizeIndex)
 		{
-			const TConstArrayView<FString> BodyPartNames = OutfitFacade.GetBodySizeBodyPartsSkeletalMeshes(BodySizeIndex);
+			const TConstArrayView<FSoftObjectPath> BodyPartNames = OutfitFacade.GetBodySizeBodyPartsSkeletalMeshPaths(BodySizeIndex);
 			OutBodyParts[BodySizeIndex].BodyParts.Reserve(BodyPartNames.Num());
 
-			for (const FString& BodyPartName : BodyPartNames)
+			for (const FSoftObjectPath& BodyPartName : BodyPartNames)
 			{
-				if (const TObjectPtr<const USkeletalMesh> SKM = LoadObject<USkeletalMesh>(nullptr, *BodyPartName))
+				if (const TObjectPtr<const USkeletalMesh> SKM = Cast<USkeletalMesh>(BodyPartName.TryLoad()))
 				{
 					OutBodyParts[BodySizeIndex].BodyParts.Add(SKM);
 				}
diff --git a/Engine/Plugins/Experimental/ChaosOutfitAsset/Source/ChaosOutfitAssetDataflowNodes/Private/ChaosOutfitAsset/GetOutfitClothCollectionsNode.cpp b/Engine/Plugins/Experimental/ChaosOutfitAsset/Source/ChaosOutfitAssetDataflowNodes/Private/ChaosOutfitAsset/GetOutfitClothCollectionsNode.cpp
--- a/Engine/Plugins/Experimental/ChaosOutfitAsset/Source/ChaosOutfitAssetDataflowNodes/Private/ChaosOutfitAsset/GetOutfitClothCollectionsNode.cpp
+++ b/Engine/Plugins/Experimental/ChaosOutfitAsset/Source/ChaosOutfitAssetDataflowNodes/Private/ChaosOutfitAsset/GetOutfitClothCollectionsNode.cpp
@@ -1,14 +1,7 @@
 // Copyright Epic Games, Inc. All Rights Reserved.
 
 #include "ChaosOutfitAsset/GetOutfitClothCollectionsNode.h"
-#include "AssetRegistry/AssetData.h"
-#include "AssetRegistry/AssetRegistryModule.h"
 #include "ChaosOutfitAsset/Outfit.h"
-#include "ChaosClothAsset/CollectionClothFacade.h"
-#include "ChaosClothAsset/CollectionClothRenderPatternFacade.h"
-#include "Engine/SkinnedAssetCommon.h"
-#include "Materials/MaterialInterface.h"
-#include "PhysicsEngine/PhysicsAsset.h"
 
 #include UE_INLINE_GENERATED_CPP_BY_NAME(GetOutfitClothCollectionsNode)
 
@@ -27,8 +20,6 @@ FChaosGetOutfitClothCollectionsNode::FChaosGetOutfitClothCollectionsNode(const U
 
 void FChaosGetOutfitClothCollectionsNode::Evaluate(UE::Dataflow::FContext& Context, const FDataflowOutput* Out) const
 {
-	using namespace UE::Chaos::ClothAsset;
-
 	if (Out->IsA(&Outfit) || Out->IsA(&ClothCollections) || Out->IsA(&NumLods))
 	{
 		TObjectPtr<const UChaosOutfit> OutOutfit;
@@ -50,79 +41,10 @@ void FChaosGetOutfitClothCollectionsNode::Evaluate(UE::Dataflow::FContext& Conte
 
 		TArray<FManagedArrayCollection> OutClothCollections;
 		OutClothCollections.Reserve(OutNumLods * OutNumPieces);
-
-		// Get the material array
-		const TArray<FSkeletalMaterial>& Materials = OutOutfit->GetMaterials();
-
-		// Fix up the fixable object paths (in case the assets have been moved) and copy the collections
-		const FAssetRegistryModule& AssetRegistryModule = FModuleManager::LoadModuleChecked<FAssetRegistryModule>(TEXT("AssetRegistry"));
-		const TArray<TSharedRef<const FManagedArrayCollection>> SelectedClothCollections = OutOutfit->GetClothCollections(InLODIndex);
-
-		for (int32 CollectionIndex = 0; CollectionIndex < SelectedClothCollections.Num(); ++CollectionIndex)
+		for (const TSharedRef<const FManagedArrayCollection>& ClothCollection : OutOutfit->GetClothCollections(InLODIndex))
 		{
-			const TSharedRef<const FManagedArrayCollection>& ClothCollection = SelectedClothCollections[CollectionIndex];
-			const TSharedRef<FManagedArrayCollection> OutClothCollection = MakeShared<FManagedArrayCollection>(*ClothCollection);
-			FCollectionClothFacade ClothFacade(OutClothCollection);
-
-			// Fix up physics asset paths
-			if (!ClothFacade.GetPhysicsAssetPathName().IsEmpty() &&
-				!AssetRegistryModule.Get().GetAssetByObjectPath(ClothFacade.GetPhysicsAssetPathName()).IsValid())
-			{
-				const UPhysicsAsset* const PhysicsAsset = OutOutfit->GetPieces()[CollectionIndex / OutNumLods].PhysicsAsset;
-				const FString PhysicsAssetPathName = PhysicsAsset ? PhysicsAsset->GetPathName() : FString();
-
-				Context.Warning(FString::Printf(TEXT("Can't find physics asset path [%s], will replace with [%s]"), *ClothFacade.GetPhysicsAssetPathName(), *PhysicsAssetPathName), this, Out);
-				ClothFacade.SetPhysicsAssetPathName(PhysicsAssetPathName);
-			}
-
-			// Fix up material paths
-			const int32 NumRenderPatterns = ClothFacade.GetNumRenderPatterns();
-			for (int32 PatternIndex = 0; PatternIndex < NumRenderPatterns; ++PatternIndex)
-			{
-				FCollectionClothRenderPatternFacade RenderPatternFacade = ClothFacade.GetRenderPattern(PatternIndex);
-				const FString& RenderMaterialPathName = RenderPatternFacade.GetRenderMaterialPathName();
-
-				if (!RenderMaterialPathName.IsEmpty() &&
-					!AssetRegistryModule.Get().GetAssetByObjectPath(RenderMaterialPathName).IsValid())
-				{
-					TArray<FString> RenderMaterialPathElements;
-					RenderMaterialPathName.ParseIntoArray(RenderMaterialPathElements, TEXT("/"), true);
-
-					// Find the closest end path name from all asset materials (currently there is no way to safely reference materials fromn the Outfit Piece)
-					int32 BestMatchLength = 0;
-					const FSkeletalMaterial* BestMatchMaterial = nullptr;
-					for (const FSkeletalMaterial& Material : Materials)
-					{
-						if (Material.MaterialInterface)
-						{
-							TArray<FString> MaterialPathElements;
-							Material.MaterialInterface->GetPathName().ParseIntoArray(MaterialPathElements, TEXT("/"), true);
-
-							for (int32 IndexFromTheEnd = 0; IndexFromTheEnd < FMath::Min(RenderMaterialPathElements.Num(), MaterialPathElements.Num()); ++IndexFromTheEnd)
-							{
-								if (RenderMaterialPathElements.Last(IndexFromTheEnd) != MaterialPathElements.Last(IndexFromTheEnd))
-								{
-									break;
-								}
-								if (BestMatchLength == IndexFromTheEnd)
-								{
-									++BestMatchLength;
-									BestMatchMaterial = &Material;
-								}
-							}
-						}
-					}
-					const FString MaterialPathName = BestMatchMaterial ? BestMatchMaterial->MaterialInterface->GetPathName() : FString();
-
-					Context.Warning(FString::Printf(TEXT("Can't find material asset path [%s], will replace with [%s]"), *RenderMaterialPathName, *MaterialPathName), this, Out);
-					RenderPatternFacade.SetRenderMaterialPathName(MaterialPathName);
-				}
-			}
-
-			// TODO: Fix up skeletal mesh paths, can't currently be done for 5.6 due to hotfix rules.
-
 			// Copy collection
-			OutClothCollections.Emplace(MoveTemp(*OutClothCollection));
+			OutClothCollections.Emplace(*ClothCollection);
 		}
 		SetValue(Context, MoveTemp(OutClothCollections), &ClothCollections);
 		SetValue(Context, OutNumLods, &NumLods);
diff --git a/Engine/Plugins/Experimental/ChaosOutfitAsset/Source/ChaosOutfitAssetEngine/Private/ChaosOutfitAsset/CollectionOutfitFacade.cpp b/Engine/Plugins/Experimental/ChaosOutfitAsset/Source/ChaosOutfitAssetEngine/Private/ChaosOutfitAsset/CollectionOutfitFacade.cpp
--- a/Engine/Plugins/Experimental/ChaosOutfitAsset/Source/ChaosOutfitAssetEngine/Private/ChaosOutfitAsset/CollectionOutfitFacade.cpp
+++ b/Engine/Plugins/Experimental/ChaosOutfitAsset/Source/ChaosOutfitAssetEngine/Private/ChaosOutfitAsset/CollectionOutfitFacade.cpp
@@ -11,6 +11,8 @@
 #include "ChaosOutfitAsset/OutfitCollection.h"
 #include "Engine/SkeletalMesh.h"
 #include "GeometryCollection/ManagedArrayCollection.h"
+#include "Serialization/Archive.h"
+#include "UObject/SoftObjectPath.h"
 
 #define LOCTEXT_NAMESPACE "CollectionOutfitFacade"
 
@@ -92,7 +94,7 @@ namespace UE::Chaos::OutfitAsset
 
 			struct
 			{
-				TManagedArrayPtr<FString> SkeletalMesh;  // Body parts skeletal mesh (several per sizes, e.g. body, head)
+				TManagedArrayPtr<FSoftObjectPath> SkeletalMeshPath;  // Body parts skeletal mesh (several per sizes, e.g. body, head)
 				TManagedArrayPtr<TArray<int32>> RBFInterpolationSampleIndices; // Precomputed RBF interpolation sample indices
 				TManagedArrayPtr<TArray<FVector3f>> RBFInterpolationSampleRestPositions; // Precomputed RBF interpolation sample rest positions
 				TManagedArrayPtr<TArray<float>> RBFInterpolationWeights; // Precomputed RBF interpolation weights
@@ -117,7 +119,7 @@ namespace UE::Chaos::OutfitAsset
 				BodySizes.Name = ManagedArrayCollection.FindAttribute<FString>(Attribute::BodySizes::Name, Group::BodySizes);
 				BodySizes.BodyPartsStart = ManagedArrayCollection.FindAttribute<int32>(Attribute::BodySizes::BodyPartsStart, Group::BodySizes);
 				BodySizes.BodyPartsCount = ManagedArrayCollection.FindAttribute<int32>(Attribute::BodySizes::BodyPartsCount, Group::BodySizes);
-				BodyParts.SkeletalMesh = ManagedArrayCollection.FindAttribute<FString>(Attribute::BodyParts::SkeletalMesh, Group::BodyParts);
+				BodyParts.SkeletalMeshPath = ManagedArrayCollection.FindAttribute<FSoftObjectPath>(Attribute::BodyParts::SkeletalMeshPath, Group::BodyParts);
 				BodyParts.RBFInterpolationSampleIndices = ManagedArrayCollection.FindAttribute<TArray<int32>>(Attribute::BodyParts::RBFInterpolationSampleIndices, Group::BodyParts);
 				BodyParts.RBFInterpolationSampleRestPositions = ManagedArrayCollection.FindAttribute<TArray<FVector3f>>(Attribute::BodyParts::RBFInterpolationSampleRestPositions, Group::BodyParts);
 				BodyParts.RBFInterpolationWeights = ManagedArrayCollection.FindAttribute<TArray<float>>(Attribute::BodyParts::RBFInterpolationWeights, Group::BodyParts);
@@ -138,7 +140,7 @@ namespace UE::Chaos::OutfitAsset
 					BodySizes.Name &&
 					BodySizes.BodyPartsStart &&
 					BodySizes.BodyPartsCount &&
-					BodyParts.SkeletalMesh &&
+					BodyParts.SkeletalMeshPath &&
 					BodyParts.RBFInterpolationSampleIndices &&
 					BodyParts.RBFInterpolationSampleRestPositions &&
 					BodyParts.RBFInterpolationWeights &&
@@ -160,7 +162,7 @@ namespace UE::Chaos::OutfitAsset
 					BodySizes.Name = &ManagedArrayCollection.AddAttribute<FString>(Attribute::BodySizes::Name, Group::BodySizes);
 					BodySizes.BodyPartsStart = &ManagedArrayCollection.AddAttribute<int32>(Attribute::BodySizes::BodyPartsStart, Group::BodySizes, Group::BodyParts);
 					BodySizes.BodyPartsCount = &ManagedArrayCollection.AddAttribute<int32>(Attribute::BodySizes::BodyPartsCount, Group::BodySizes);
-					BodyParts.SkeletalMesh = &ManagedArrayCollection.AddAttribute<FString>(Attribute::BodyParts::SkeletalMesh, Group::BodyParts);
+					BodyParts.SkeletalMeshPath = &ManagedArrayCollection.AddAttribute<FSoftObjectPath>(Attribute::BodyParts::SkeletalMeshPath, Group::BodyParts);
 					BodyParts.RBFInterpolationSampleIndices = &ManagedArrayCollection.AddAttribute<TArray<int32>>(Attribute::BodyParts::RBFInterpolationSampleIndices, Group::BodyParts);
 					BodyParts.RBFInterpolationSampleRestPositions = &ManagedArrayCollection.AddAttribute<TArray<FVector3f>>(Attribute::BodyParts::RBFInterpolationSampleRestPositions, Group::BodyParts);
 					BodyParts.RBFInterpolationWeights = &ManagedArrayCollection.AddAttribute<TArray<float>>(Attribute::BodyParts::RBFInterpolationWeights, Group::BodyParts);
@@ -177,6 +179,36 @@ namespace UE::Chaos::OutfitAsset
 				ManagedArrayCollection.EmptyGroup(Group::BodyParts);
 				ManagedArrayCollection.EmptyGroup(Group::Measurements);
 			}
+
+			void PostSerialize(const FArchive& Ar)
+			{
+				if (Ar.IsLoading())
+				{
+					auto ConvertStringToSoftObjectPath = [](const TManagedArray<FString>& StringArray, TManagedArray<FSoftObjectPath>& SoftObjectPathArray)
+						{
+							check(StringArray.Num() == SoftObjectPathArray.Num());
+							for (int32 Index = 0; Index < StringArray.Num(); ++Index)
+							{
+								SoftObjectPathArray[Index] = FSoftObjectPath(StringArray[Index]);
+								SoftObjectPathArray[Index].PreSavePath(); // This will fixup the path with redirects.
+							}
+						};
+
+					// BodyParts::SkeletalMesh -> BodyParts::SkeletalMeshPath
+					using namespace OutfitCollection;
+					PRAGMA_DISABLE_DEPRECATION_WARNINGS
+					if (TManagedArrayPtr<FString> SkeletalMeshes = ManagedArrayCollection.FindAttribute<FString>(Attribute::BodyParts::SkeletalMesh, Group::BodyParts))
+					{
+						if (!BodyParts.SkeletalMeshPath)
+						{
+							BodyParts.SkeletalMeshPath = &ManagedArrayCollection.AddAttribute<FSoftObjectPath>(Attribute::BodyParts::SkeletalMeshPath, Group::BodyParts);
+						}
+						ConvertStringToSoftObjectPath(*SkeletalMeshes, *BodyParts.SkeletalMeshPath);
+						ManagedArrayCollection.RemoveAttribute(Attribute::BodyParts::SkeletalMesh, Group::BodyParts);
+					}
+					PRAGMA_ENABLE_DEPRECATION_WARNINGS
+				}
+			}
 		};
 
 		class FBodyMatchParameters
@@ -400,15 +432,15 @@ namespace UE::Chaos::OutfitAsset
 
 		for (int32 BodySize = 0; BodySize < GetNumBodySizes(); ++BodySize)
 		{
-			const TConstArrayView<FString> BodyPartsSkeletalMeshes = GetBodySizeBodyPartsSkeletalMeshes(BodySize);
-			if (!Algo::AnyOf(BodyPartsSkeletalMeshes, [&AssetRegistryModule, bBodyPartMustExist](const FString& BodyPartsSkeletalMesh)
+			const TConstArrayView<FSoftObjectPath> BodyPartsSkeletalMeshes = GetBodySizeBodyPartsSkeletalMeshPaths(BodySize);
+			if (!Algo::AnyOf(BodyPartsSkeletalMeshes, [&AssetRegistryModule, bBodyPartMustExist](const FSoftObjectPath& BodyPartsSkeletalMesh)
 				{
 					if (bBodyPartMustExist)
 					{
 						const FAssetData AssetData = AssetRegistryModule.Get().GetAssetByObjectPath(BodyPartsSkeletalMesh);
 						return AssetData.IsValid();  //  Body part exists
 					}
-					return !BodyPartsSkeletalMesh.IsEmpty();  // Body part is not empty
+					return BodyPartsSkeletalMesh.IsValid();  // Body part is not empty
 				}))
 			{
 				continue;  // Invalid body parts, try another size
@@ -442,18 +474,30 @@ namespace UE::Chaos::OutfitAsset
 		return false;
 	}
 
-	TArray<FString> FCollectionOutfitConstFacade::GetOutfitBodyPartsSkeletalMeshes() const
+	TArray<FSoftObjectPath> FCollectionOutfitConstFacade::GetOutfitBodyPartsSkeletalMeshPaths() const
 	{
 		check(IsValid());
-		TArray<FString> BodyParts;
-		BodyParts.Reserve(GetOutfitCollection()->BodyParts.SkeletalMesh->Num());
-		for (const FString& SkeletalMesh : (*GetOutfitCollection()->BodyParts.SkeletalMesh))
+		TArray<FSoftObjectPath> BodyParts;
+		BodyParts.Reserve(GetOutfitCollection()->BodyParts.SkeletalMeshPath->Num());
+		for (const FSoftObjectPath& SkeletalMesh : (*GetOutfitCollection()->BodyParts.SkeletalMeshPath))
 		{
 			BodyParts.AddUnique(SkeletalMesh);
 		}
 		return BodyParts;
 	}
 
+	TArray<FString> FCollectionOutfitConstFacade::GetOutfitBodyPartsSkeletalMeshes() const
+	{
+		const TArray<FSoftObjectPath> BodyPartPaths = GetOutfitBodyPartsSkeletalMeshPaths();
+		TArray<FString> BodyParts;
+		BodyParts.Reserve(BodyPartPaths.Num());
+		for (const FSoftObjectPath& SkeletalMesh : BodyPartPaths)
+		{
+			BodyParts.Emplace(SkeletalMesh.ToString());
+		}
+		return BodyParts;
+	}
+
 	TArray<int32> FCollectionOutfitConstFacade::GetOutfitBodySizes(const FGuid& Guid) const
 	{
 		check(IsValid());
@@ -601,7 +645,7 @@ namespace UE::Chaos::OutfitAsset
 			(*GetOutfitCollection()->BodySizes.Name)[BodySize] : EmptyString;
 	}
 
-	TConstArrayView<FString> FCollectionOutfitConstFacade::GetBodySizeBodyPartsSkeletalMeshes(int32 BodySize) const
+	TConstArrayView<FSoftObjectPath> FCollectionOutfitConstFacade::GetBodySizeBodyPartsSkeletalMeshPaths(int32 BodySize) const
 	{
 		check(IsValid());
 		if (GetOutfitCollection()->BodySizes.Name->IsValidIndex(BodySize))
@@ -609,13 +653,13 @@ namespace UE::Chaos::OutfitAsset
 			const int32 BodyPartsStart = (*GetOutfitCollection()->BodySizes.BodyPartsStart)[BodySize];
 			const int32 BodyPartsCount = (*GetOutfitCollection()->BodySizes.BodyPartsCount)[BodySize];
 			if (BodyPartsCount &&
-				GetOutfitCollection()->BodyParts.SkeletalMesh->IsValidIndex(BodyPartsStart) &&
-				GetOutfitCollection()->BodyParts.SkeletalMesh->IsValidIndex(BodyPartsStart + BodyPartsCount - 1))
+				GetOutfitCollection()->BodyParts.SkeletalMeshPath->IsValidIndex(BodyPartsStart) &&
+				GetOutfitCollection()->BodyParts.SkeletalMeshPath->IsValidIndex(BodyPartsStart + BodyPartsCount - 1))
 			{
-				return TConstArrayView<FString>(GetOutfitCollection()->BodyParts.SkeletalMesh->GetData() + BodyPartsStart, BodyPartsCount);
+				return TConstArrayView<FSoftObjectPath>(GetOutfitCollection()->BodyParts.SkeletalMeshPath->GetData() + BodyPartsStart, BodyPartsCount);
 			}
 		}
-		return TConstArrayView<FString>();
+		return TConstArrayView<FSoftObjectPath>();
 	}
 	
 	int32 FCollectionOutfitConstFacade::GetBodySizeBodyPartOffset(const int32 BodySize) const
@@ -689,6 +733,11 @@ namespace UE::Chaos::OutfitAsset
 	{
 	}
 
+	void FCollectionOutfitFacade::PostSerialize(const FArchive& Ar)
+	{
+		GetOutfitCollection()->PostSerialize(Ar);
+	}
+
 	FCollectionOutfitFacade::~FCollectionOutfitFacade() = default;
 
 	void FCollectionOutfitFacade::DefineSchema()
@@ -703,7 +752,7 @@ namespace UE::Chaos::OutfitAsset
 		check(IsValid());
 	}
 
-	int32 FCollectionOutfitFacade::AddBodySize(const FString& Name, const TConstArrayView<FString> BodyPartsSkeletalMeshes, const TMap<FString, float>& Measurements, const FRBFInterpolationDataWrapper& InterpolationData)
+	int32 FCollectionOutfitFacade::AddBodySize(const FString& Name, const TConstArrayView<FSoftObjectPath> BodyPartsSkeletalMeshes, const TMap<FString, float>& Measurements, const FRBFInterpolationDataWrapper& InterpolationData)
 	{
 		using namespace OutfitCollection;
 		check(IsValid());
@@ -733,7 +782,7 @@ namespace UE::Chaos::OutfitAsset
 				
 			for (int32 Index = 0; Index < BodyPartsCount; ++Index)
 			{
-				(*GetOutfitCollection()->BodyParts.SkeletalMesh)[BodyPartsStart + Index] = BodyPartsSkeletalMeshes[Index];
+				(*GetOutfitCollection()->BodyParts.SkeletalMeshPath)[BodyPartsStart + Index] = BodyPartsSkeletalMeshes[Index];
 				if (bInterpolationDataValid)
 				{
 					(*GetOutfitCollection()->BodyParts.RBFInterpolationSampleIndices)[BodyPartsStart + Index] = InterpolationData.SampleIndices[Index];
@@ -774,6 +823,18 @@ namespace UE::Chaos::OutfitAsset
 		return BodySize;
 	}
 
+	int32 FCollectionOutfitFacade::AddBodySize(const FString& Name, const TConstArrayView<FString> BodyPartsSkeletalMeshes, const TMap<FString, float>& Measurements, const FRBFInterpolationDataWrapper& InterpolationData)
+	{
+		TArray<FSoftObjectPath> SKMPaths;
+		SKMPaths.Reserve(BodyPartsSkeletalMeshes.Num());
+		for (const FString& SKM : BodyPartsSkeletalMeshes)
+		{
+			SKMPaths.Emplace(SKM);
+		}
+
+		return AddBodySize(Name, SKMPaths, Measurements, InterpolationData);
+	}
+
 	void FCollectionOutfitFacade::AddOutfit(const FGuid& Guid, const int32 BodySize, const FString& Name, const TMap<FGuid, FString>& Pieces)
 	{
 		using namespace OutfitCollection;
@@ -879,7 +940,7 @@ namespace UE::Chaos::OutfitAsset
 			// Add or replace the named body size
 			RemappedBodySize[BodySize - SelectedBodySize] = AddBodySize(
 				Other.GetBodySizeName(BodySize),
-				Other.GetBodySizeBodyPartsSkeletalMeshes(BodySize),
+				Other.GetBodySizeBodyPartsSkeletalMeshPaths(BodySize),
 				Other.GetBodySizeMeasurements(BodySize),
 				Other.GetBodySizeInterpolationData(BodySize));
 		}
diff --git a/Engine/Plugins/Experimental/ChaosOutfitAsset/Source/ChaosOutfitAssetEngine/Private/ChaosOutfitAsset/Outfit.cpp b/Engine/Plugins/Experimental/ChaosOutfitAsset/Source/ChaosOutfitAssetEngine/Private/ChaosOutfitAsset/Outfit.cpp
--- a/Engine/Plugins/Experimental/ChaosOutfitAsset/Source/ChaosOutfitAssetEngine/Private/ChaosOutfitAsset/Outfit.cpp
+++ b/Engine/Plugins/Experimental/ChaosOutfitAsset/Source/ChaosOutfitAssetEngine/Private/ChaosOutfitAsset/Outfit.cpp
@@ -604,6 +604,10 @@ bool FChaosOutfitPiece::Serialize(FArchive& Ar)
 		{
 			TSharedRef<FManagedArrayCollection> NewCollection = MakeShared<FManagedArrayCollection>();
 			NewCollection->Serialize(Ar);
+
+			UE::Chaos::ClothAsset::FCollectionClothFacade ClothFacade(NewCollection);
+			ClothFacade.PostSerialize(Ar);
+
 			Collections.Emplace(MoveTemp(NewCollection));
 		}
 	}
@@ -721,7 +725,7 @@ void UChaosOutfit::Add(const FChaosSizedOutfitSource& SizedOutfitSource, const F
 	const FString BodySizeName = SizedOutfitSource.GetBodySizeName();
 	if (!BodySizeName.IsEmpty())
 	{
-		TArray<FString> BodyPartsSkeletalMeshes;
+		TArray<FSoftObjectPath> BodyPartsSkeletalMeshes;
 		TArray<TArray<int32>> InterpolationData_SampleIndices;
 		TArray<TArray<FVector3f>> InterpolationData_SampleRestPositions;
 		TArray<TArray<float>> InterpolationData_InterpolationWeights;
@@ -889,6 +893,11 @@ void UChaosOutfit::Serialize(FArchive& Ar)
 
 	Ar << ReferenceSkeleton;
 	Ar << *RenderData;
+	if (Ar.IsLoading())
+	{
+		UE::Chaos::OutfitAsset::FCollectionOutfitFacade OutfitFacade(OutfitCollection);
+		OutfitFacade.PostSerialize(Ar);
+	}
 }
 
 template<typename RenderDataType UE_REQUIRES_DEFINITION(std::is_same_v<RenderDataType, FSkeletalMeshRenderData> || std::is_same_v<RenderDataType, UChaosOutfit::FRenderData>)>
diff --git a/Engine/Plugins/Experimental/ChaosOutfitAsset/Source/ChaosOutfitAssetEngine/Private/ChaosOutfitAsset/OutfitAsset.cpp b/Engine/Plugins/Experimental/ChaosOutfitAsset/Source/ChaosOutfitAssetEngine/Private/ChaosOutfitAsset/OutfitAsset.cpp
--- a/Engine/Plugins/Experimental/ChaosOutfitAsset/Source/ChaosOutfitAssetEngine/Private/ChaosOutfitAsset/OutfitAsset.cpp
+++ b/Engine/Plugins/Experimental/ChaosOutfitAsset/Source/ChaosOutfitAssetEngine/Private/ChaosOutfitAsset/OutfitAsset.cpp
@@ -78,7 +78,7 @@ void UChaosOutfitAsset::Build(const TObjectPtr<const UChaosOutfit> InOutfit, UE:
 		if (ContextAssetStore)
 		{
 			// Fix up resized material paths
-			TMap<FString, FString> MaterialPathsToFixUp;
+			TMap<FSoftObjectPath, FSoftObjectPath> MaterialPathsToFixUp;
 			MaterialPathsToFixUp.Reserve(Materials.Num());
 
 			for (int32 MaterialIndex = 0; MaterialIndex < Materials.Num(); ++MaterialIndex)
@@ -112,18 +112,18 @@ void UChaosOutfitAsset::Build(const TObjectPtr<const UChaosOutfit> InOutfit, UE:
 								TOptional<FCollectionClothFacade> FixedUpClothFacade;
 
 								FCollectionClothConstFacade ClothFacade(Collection);
-								const TConstArrayView<FString> RenderMaterialPathNames = ClothFacade.GetRenderMaterialPathName();
+								const TConstArrayView<FSoftObjectPath> RenderMaterialPathNames = ClothFacade.GetRenderMaterialSoftObjectPathName();
 								for (int32 PathIndex = 0; PathIndex < RenderMaterialPathNames.Num(); ++PathIndex)
 								{
-									const FString& RenderMaterialPathName = RenderMaterialPathNames[PathIndex];
-									if (const FString* MaterialPathToFixUp = MaterialPathsToFixUp.Find(RenderMaterialPathName))
+									const FSoftObjectPath& RenderMaterialPathName = RenderMaterialPathNames[PathIndex];
+									if (const FSoftObjectPath* MaterialPathToFixUp = MaterialPathsToFixUp.Find(RenderMaterialPathName))
 									{
 										if (!FixedUpCollection.IsSet())
 										{
 											FixedUpCollection.Emplace(MakeShared<FManagedArrayCollection>(*Collection));
 											FixedUpClothFacade.Emplace(FixedUpCollection.GetValue());
 										}
-										TArrayView<FString> FixedUpRenderMaterialPathNames = FixedUpClothFacade->GetRenderMaterialPathName();
+										TArrayView<FSoftObjectPath> FixedUpRenderMaterialPathNames = FixedUpClothFacade->GetRenderMaterialSoftObjectPathName();
 										check(FixedUpRenderMaterialPathNames[PathIndex] == RenderMaterialPathName);
 							
 										FixedUpRenderMaterialPathNames[PathIndex] = *MaterialPathToFixUp;
@@ -167,11 +167,11 @@ void UChaosOutfitAsset::Build(const TObjectPtr<const UChaosOutfit> InOutfit, UE:
 
 	// Populate the body sizes
 	const FCollectionOutfitConstFacade OutfitFacade(OutfitCollection);
-	const TArray<FString> OutfitBodyPartsSkeletalMeshes = OutfitFacade.GetOutfitBodyPartsSkeletalMeshes();
+	const TArray<FSoftObjectPath> OutfitBodyPartsSkeletalMeshes = OutfitFacade.GetOutfitBodyPartsSkeletalMeshPaths();
 	Bodies.Reset(OutfitBodyPartsSkeletalMeshes.Num());
-	for (const FString& OutfitBodyPartsSkeletalMesh : OutfitBodyPartsSkeletalMeshes)
+	for (const FSoftObjectPath& OutfitBodyPartsSkeletalMesh : OutfitBodyPartsSkeletalMeshes)
 	{
-		if (USkeletalMesh* const Body = LoadObject<USkeletalMesh>(this, *OutfitBodyPartsSkeletalMesh))
+		if (USkeletalMesh* const Body = Cast<USkeletalMesh>(OutfitBodyPartsSkeletalMesh.TryLoad()))
 		{
 			Bodies.Emplace(Body);
 		}
@@ -227,6 +227,11 @@ void UChaosOutfitAsset::Serialize(FArchive& Ar)
 			GetResourceForRendering()->Serialize(Ar, this);
 		}
 	}
+	if (Ar.IsLoading())
+	{
+		UE::Chaos::OutfitAsset::FCollectionOutfitFacade OutfitFacade(OutfitCollection);
+		OutfitFacade.PostSerialize(Ar);
+	}
 }
 
 void UChaosOutfitAsset::PostLoad()
diff --git a/Engine/Plugins/Experimental/ChaosOutfitAsset/Source/ChaosOutfitAssetEngine/Private/ChaosOutfitAsset/OutfitCollection.cpp b/Engine/Plugins/Experimental/ChaosOutfitAsset/Source/ChaosOutfitAssetEngine/Private/ChaosOutfitAsset/OutfitCollection.cpp
--- a/Engine/Plugins/Experimental/ChaosOutfitAsset/Source/ChaosOutfitAssetEngine/Private/ChaosOutfitAsset/OutfitCollection.cpp
+++ b/Engine/Plugins/Experimental/ChaosOutfitAsset/Source/ChaosOutfitAssetEngine/Private/ChaosOutfitAsset/OutfitCollection.cpp
@@ -46,6 +46,8 @@ namespace UE::Chaos::OutfitAsset
 			// BodyParts group
 			namespace BodyParts
 			{
+				const FName SkeletalMeshPath(TEXT("SkeletalMeshPath"));  // FSoftObjectPath
+				UE_DEPRECATED(5.7, "Use SkeletalMeshPath instead.")
 				const FName SkeletalMesh(TEXT("SkeletalMesh"));  // FString
 				const FName RBFInterpolationSampleIndices(TEXT("RBFInterpolationSampleIndices"));  // TArray<int32>
 				const FName RBFInterpolationSampleRestPositions(TEXT("RBFInterpolationSampleRestPositions"));  // TArray<FVector3f>
diff --git a/Engine/Plugins/Experimental/ChaosOutfitAsset/Source/ChaosOutfitAssetEngine/Public/ChaosOutfitAsset/CollectionOutfitFacade.h b/Engine/Plugins/Experimental/ChaosOutfitAsset/Source/ChaosOutfitAssetEngine/Public/ChaosOutfitAsset/CollectionOutfitFacade.h
--- a/Engine/Plugins/Experimental/ChaosOutfitAsset/Source/ChaosOutfitAssetEngine/Public/ChaosOutfitAsset/CollectionOutfitFacade.h
+++ b/Engine/Plugins/Experimental/ChaosOutfitAsset/Source/ChaosOutfitAssetEngine/Public/ChaosOutfitAsset/CollectionOutfitFacade.h
@@ -54,6 +54,8 @@ namespace UE::Chaos::OutfitAsset
 		UE_API bool HasValidBodySize(bool bBodyPartMustExist = false, bool bMeasurementsMustExist = false, bool bInterpolationDataMustExist = false) const;
 
 		/** Return all unique body part skeletal meshes used by this outfit. */
+		UE_API TArray<FSoftObjectPath> GetOutfitBodyPartsSkeletalMeshPaths() const;
+		UE_DEPRECATED(5.7, "Use GetOutfitBodyPartsSkeletalMeshPaths instead.")
 		UE_API TArray<FString> GetOutfitBodyPartsSkeletalMeshes() const;
 
 		/** Return the body size available for the specified outfit GUID. */
@@ -90,7 +92,12 @@ namespace UE::Chaos::OutfitAsset
 		UE_API const FString& GetBodySizeName(const int32 BodySize) const;
 
 		/** Return the body part skeletal meshes path name for the specified body size. */
-		UE_API TConstArrayView<FString> GetBodySizeBodyPartsSkeletalMeshes(const int32 BodySize) const;
+		UE_API TConstArrayView<FSoftObjectPath> GetBodySizeBodyPartsSkeletalMeshPaths(const int32 BodySize) const;
+		UE_DEPRECATED(5.7, "Use GetBodySizeBodyPartsSkeletalMeshPaths instead.")
+		TConstArrayView<FString> GetBodySizeBodyPartsSkeletalMeshes(const int32 BodySize) const
+		{
+			return TConstArrayView<FString>();
+		}
 
 		/** Return the body size body part offset (to index into the global body part array). Returns INDEX_NONE if BodySize is invalid.*/
 		UE_API int32 GetBodySizeBodyPartOffset(const int32 BodySize) const;
@@ -129,6 +136,9 @@ namespace UE::Chaos::OutfitAsset
 		FCollectionOutfitFacade& operator=(FCollectionOutfitFacade&&) = default;
 		UE_API virtual ~FCollectionOutfitFacade() override;
 
+		/** Do post-serialization fixups */
+		UE_API void PostSerialize(const FArchive& Ar);
+
 		/** Add the Selection attributes to the underlying collection. */
 		UE_API void DefineSchema();
 
@@ -139,6 +149,12 @@ namespace UE::Chaos::OutfitAsset
 		 * Add a new body size and returns its index.
 		 * Any pre-existing sizes data will be replaced.
 		 */
+		UE_API int32 AddBodySize(
+			const FString& Name,
+			const TConstArrayView<FSoftObjectPath> BodyPartsSkeletalMeshes,
+			const TMap<FString, float>& Measurements,
+			const FRBFInterpolationDataWrapper& InterpolationData);
+		UE_DEPRECATED(5.7, "Use version with FSoftObjectPath BodyPartsSkeletalMeshes")
 		UE_API int32 AddBodySize(
 			const FString& Name,
 			const TConstArrayView<FString> BodyPartsSkeletalMeshes,
@@ -150,7 +166,7 @@ namespace UE::Chaos::OutfitAsset
 		{
 			const int32 BodySize = FindBodySize(Name);
 
-			return BodySize != INDEX_NONE ? BodySize : AddBodySize(Name, {}, {}, {});
+			return BodySize != INDEX_NONE ? BodySize : AddBodySize(Name, TConstArrayView<FSoftObjectPath>(), {}, {});
 		}
 
 		/**
diff --git a/Engine/Plugins/Experimental/ChaosOutfitAsset/Source/ChaosOutfitAssetEngine/Public/ChaosOutfitAsset/OutfitCollection.h b/Engine/Plugins/Experimental/ChaosOutfitAsset/Source/ChaosOutfitAssetEngine/Public/ChaosOutfitAsset/OutfitCollection.h
--- a/Engine/Plugins/Experimental/ChaosOutfitAsset/Source/ChaosOutfitAssetEngine/Public/ChaosOutfitAsset/OutfitCollection.h
+++ b/Engine/Plugins/Experimental/ChaosOutfitAsset/Source/ChaosOutfitAssetEngine/Public/ChaosOutfitAsset/OutfitCollection.h
@@ -44,6 +44,8 @@ namespace UE::Chaos::OutfitAsset
 			}
 			namespace BodyParts
 			{
+				UE_API extern const FName SkeletalMeshPath;  // FSoftObjectPath
+				UE_DEPRECATED(5.7, "Use SkeletalMeshPath instead.")
 				UE_API extern const FName SkeletalMesh;  // FString
 				UE_API extern const FName RBFInterpolationSampleIndices;  // TArray<int32>
 				UE_API extern const FName RBFInterpolationSampleRestPositions;  // TArray<FVector3f>
diff --git a/Engine/Source/Runtime/Experimental/Chaos/Public/GeometryCollection/ManagedArrayTypeValues.inl b/Engine/Source/Runtime/Experimental/Chaos/Public/GeometryCollection/ManagedArrayTypeValues.inl
--- a/Engine/Source/Runtime/Experimental/Chaos/Public/GeometryCollection/ManagedArrayTypeValues.inl
+++ b/Engine/Source/Runtime/Experimental/Chaos/Public/GeometryCollection/ManagedArrayTypeValues.inl
@@ -63,6 +63,7 @@ MANAGED_ARRAY_TYPE(FUintVector2, UintVector2)
 MANAGED_ARRAY_TYPE(TObjectPtr<UObject>, UObjectArray)
 MANAGED_ARRAY_TYPE(Chaos::FLinearCurve, LinearCurve)
 MANAGED_ARRAY_TYPE(FName, Name)
+MANAGED_ARRAY_TYPE(FSoftObjectPath, SoftObjectPaath)
 // NOTE: new types must be added at the bottom to keep serialization from breaking
 
 
