Chaos Cloth Asset - Changed asset colors to more visible colors.

[RN] Chaos Cloth Asset - Changed the Cloth and Outfit Asset colors to better differentiate them in the asset browser.

diff --git a/Engine/Plugins/ChaosClothAssetEditor/Source/ChaosClothAssetTools/Public/ChaosClothAsset/ColorScheme.h b/Engine/Plugins/ChaosClothAssetEditor/Source/ChaosClothAssetTools/Public/ChaosClothAsset/ColorScheme.h
--- a/Engine/Plugins/ChaosClothAssetEditor/Source/ChaosClothAssetTools/Public/ChaosClothAsset/ColorScheme.h
+++ b/Engine/Plugins/ChaosClothAssetEditor/Source/ChaosClothAssetTools/Public/ChaosClothAsset/ColorScheme.h
@@ -8,7 +8,7 @@ namespace UE::Chaos::ClothAsset
 {
 	struct FColorScheme
 	{
-		static inline const FLinearColor Asset = FColor(180, 120, 110);
+		static inline const FLinearColor Asset = FColor(242, 77, 26);  // Orange
 		static inline const FLinearColor NodeHeader = FColor(180, 120, 110);
 		static inline const FLinearColor NodeBody = FColor(18, 12, 11, 127);
 	};
diff --git a/Engine/Plugins/Experimental/ChaosOutfitAsset/Source/ChaosOutfitAssetEditor/Private/ChaosOutfitAsset/AssetDefinition_OutfitAsset.cpp b/Engine/Plugins/Experimental/ChaosOutfitAsset/Source/ChaosOutfitAssetEditor/Private/ChaosOutfitAsset/AssetDefinition_OutfitAsset.cpp
--- a/Engine/Plugins/Experimental/ChaosOutfitAsset/Source/ChaosOutfitAssetEditor/Private/ChaosOutfitAsset/AssetDefinition_OutfitAsset.cpp
+++ b/Engine/Plugins/Experimental/ChaosOutfitAsset/Source/ChaosOutfitAssetEditor/Private/ChaosOutfitAsset/AssetDefinition_OutfitAsset.cpp
@@ -10,7 +10,7 @@
 
 namespace UE::Chaos::OutfitAsset::Private
 {
-	static const FLinearColor OutfitAssetColor = FColor(162, 108, 99);
+	static const FLinearColor OutfitAssetColor = FColor(243, 244, 11);  // Bright yellowx
 }
 
 FText UAssetDefinition_OutfitAsset::GetAssetDisplayName() const
