Chaos Outfit Asset - Fixed a crash in the Outfit skeleton merge when different roots are used by diffrent Cloth Assets.

[RN] Chaos Outfit Asset - Fixed a crash in the Outfit skeleton merge when different roots are used by diffrent Cloth Assets.

diff --git a/Engine/Plugins/Experimental/ChaosOutfitAsset/Source/ChaosOutfitAssetEngine/Private/ChaosOutfitAsset/Outfit.cpp b/Engine/Plugins/Experimental/ChaosOutfitAsset/Source/ChaosOutfitAssetEngine/Private/ChaosOutfitAsset/Outfit.cpp
--- a/Engine/Plugins/Experimental/ChaosOutfitAsset/Source/ChaosOutfitAssetEngine/Private/ChaosOutfitAsset/Outfit.cpp
+++ b/Engine/Plugins/Experimental/ChaosOutfitAsset/Source/ChaosOutfitAssetEngine/Private/ChaosOutfitAsset/Outfit.cpp
@@ -130,47 +130,61 @@ struct UChaosOutfit::FRenderData
 namespace UE::Chaos::OutfitAsset::Private
 {
 	// Merge bone to a new skeleton and return its new index in the new skeleton
-	// returns INDEX_NONE oif the parent don't match
 	static int32 MergeBoneToSkeleton(const int32 BoneIndex, const FReferenceSkeleton& InSkeleton, FReferenceSkeleton& InOutSkeleton)
 	{
-		// Check if it exists in the new skeleton with the same parent
 		const FName BoneName = InSkeleton.GetBoneName(BoneIndex);
 		const int32 ParentBoneIndex = InSkeleton.GetParentIndex(BoneIndex);
 		const int32 ExistingBoneIndex = InOutSkeleton.FindBoneIndex(BoneName);
-		if (ExistingBoneIndex != INDEX_NONE ||  // Already exists or
-			ParentBoneIndex == INDEX_NONE)      // it is the root bone, update the bone map
+
+		if (ParentBoneIndex == INDEX_NONE)  // Merging root bones
 		{
-			// Check if the parents matches
-			const int32 ExistingParentBoneIndex = (ExistingBoneIndex != INDEX_NONE) ? InOutSkeleton.GetParentIndex(ExistingBoneIndex) : INDEX_NONE;
-			if (!ensureMsgf((ParentBoneIndex == INDEX_NONE && ExistingParentBoneIndex == INDEX_NONE) ||
-				(InSkeleton.GetBoneName(ParentBoneIndex) == InOutSkeleton.GetBoneName(ExistingParentBoneIndex)),
-				TEXT("Bone parent must match")))
+			constexpr int32 RootBoneIndex = 0;  // A root bone always exists at index 0 (actually created in the UChaosOutfit::Init() function)
+			const FName ExistingRootBoneName = InOutSkeleton.GetBoneName(RootBoneIndex);
+
+			if (ExistingBoneIndex != INDEX_NONE)  // This root bone has already been merged
 			{
-				// TODO Better error management than ensure
+				UE_CLOG(ExistingRootBoneName != BoneName,  // Warn if this bone is not the root in the existing hierarchy
+					LogChaosOutfitAsset, Warning, TEXT("Root bone [%s] is no longer root bone in the merged hierarchy."), *BoneName.ToString());
+
+				return ExistingBoneIndex;
 			}
-			// TODO: Warn if the refposes don't match
+			else  // The existing root bone has a different name and must be renamed (this also allows for the UChaosOutfit::Init() "Root" name to be changed if required)
+			{
+				FReferenceSkeletonModifier ReferenceSkeletonModifier(InOutSkeleton, nullptr);
+				ReferenceSkeletonModifier.Rename(ExistingRootBoneName, BoneName);
+
+				UE_LOG(LogChaosOutfitAsset, Display, TEXT("Root bone name updated from [%s] to [%s]."), *ExistingRootBoneName.ToString(), *BoneName.ToString());
 
-			// Update the bone map
-			return ExistingBoneIndex;
+				return RootBoneIndex;
+			}
 		}
-		else  // Doesn't already exists, add the bone
+		else  // Merging non root bones
 		{
-			// Find its parent in the new skeleton
-			const FName ParentBoneName = (ParentBoneIndex == INDEX_NONE) ? NAME_None : InSkeleton.GetBoneName(ParentBoneIndex);
+			const FName ParentBoneName = InSkeleton.GetBoneName(ParentBoneIndex);
 			const int32 ExistingParentBoneIndex = InOutSkeleton.FindBoneIndex(ParentBoneName);
-			checkf(ParentBoneIndex == INDEX_NONE || ExistingParentBoneIndex != INDEX_NONE, TEXT("Either root or a valid parent should exist at this point."));
-			FReferenceSkeletonModifier ReferenceSkeletonModifier(InOutSkeleton, nullptr);
+			checkf(ExistingParentBoneIndex != INDEX_NONE, TEXT("Parent bones must always be merged first before their children."));
 
-			FMeshBoneInfo MeshBoneInfo;
+			if (ExistingBoneIndex != INDEX_NONE)  // This bone has already been merged
+			{
+				UE_CLOG(ExistingParentBoneIndex != InOutSkeleton.GetParentIndex(ExistingBoneIndex),  // Warn if the found bone parent's position doesn't match the original parent bone name's position
+					LogChaosOutfitAsset, Warning, TEXT("Bone [%s] has already been merged to an Outfit but with a different hierarchy for its parent [%s]."), *BoneName.ToString(), *ParentBoneName.ToString());
+
+				return ExistingBoneIndex;
+			}
+			else  // This bone hasn't been merged yet, add it as a new bone under its merged parent
+			{
+				FReferenceSkeletonModifier ReferenceSkeletonModifier(InOutSkeleton, nullptr);
+
+				FMeshBoneInfo MeshBoneInfo;
 #if WITH_EDITORONLY_DATA
-			MeshBoneInfo.ExportName = BoneName.ToString();
+				MeshBoneInfo.ExportName = BoneName.ToString();
 #endif
-			MeshBoneInfo.Name = BoneName;
-			MeshBoneInfo.ParentIndex = ExistingParentBoneIndex;
-			ReferenceSkeletonModifier.Add(MeshBoneInfo, InSkeleton.GetRefBonePose()[BoneIndex]);
+				MeshBoneInfo.Name = BoneName;
+				MeshBoneInfo.ParentIndex = ExistingParentBoneIndex;
+				ReferenceSkeletonModifier.Add(MeshBoneInfo, InSkeleton.GetRefBonePose()[BoneIndex]);
 
-			// Update the bone map
-			return ReferenceSkeletonModifier.FindBoneIndex(BoneName);
+				return ReferenceSkeletonModifier.FindBoneIndex(BoneName);
+			}
 		}
 	}
 
@@ -188,7 +202,7 @@ namespace UE::Chaos::OutfitAsset::Private
 		return Children;
 	}
 
-	static void LogHierachy(const TArray<FMeshBoneInfo>& MeshBoneInfos, const int32 ParentIndex = INDEX_NONE, int32 Indent = 0)
+	static void LogHierarchy(const TArray<FMeshBoneInfo>& MeshBoneInfos, const int32 ParentIndex = INDEX_NONE, int32 Indent = 0)
 	{
 		auto MakeIndentString = [](const int32 Indent) -> FString
 			{
@@ -208,7 +222,7 @@ namespace UE::Chaos::OutfitAsset::Private
 
 			UE_LOG(LogChaosOutfitAsset, VeryVerbose, TEXT("%d - %s"), ChildIndex, *IndentedName);
 			
-			LogHierachy(MeshBoneInfos, ChildIndex, Indent + 1);
+			LogHierarchy(MeshBoneInfos, ChildIndex, Indent + 1);
 		}
 	}
 
@@ -245,8 +259,8 @@ namespace UE::Chaos::OutfitAsset::Private
 			BonesToProcess[BoneIndex] = false;
 		}
 
-		UE_LOG(LogChaosOutfitAsset, VeryVerbose, TEXT("------- Outfit Reference Skseleton Merging --------"));
-		LogHierachy(InOutSkeleton.GetRefBoneInfo());
+		UE_LOG(LogChaosOutfitAsset, VeryVerbose, TEXT("-------- Outfit Reference Skeleton Merging --------"));
+		LogHierarchy(InOutSkeleton.GetRefBoneInfo());
 		UE_LOG(LogChaosOutfitAsset, VeryVerbose, TEXT("---------------------------------------------------"));
 		for (int32 Index = 0; Index < OutBoneMap.Num(); ++Index)
 		{
@@ -1203,7 +1217,7 @@ void UChaosOutfit::Merge(
 			FChaosOutfitPiece& Piece = OutPieces.Add_GetRef(Private::Clone(Pieces, PieceIndex));
 			Piece.RemapBoneIndices(BoneMap);
 			// TODO: Fix case where a piece already exists in the outfit as to not mess up the LODRenderData piece indices
-			UE_CLOG(AssetGuids.Contains(Piece.AssetGuid), LogChaosOutfitAsset, Warning, TEXT("Piece [%s] already exists in this Outfit!."), *Piece.Name.ToString());
+			UE_CLOG(AssetGuids.Contains(Piece.AssetGuid), LogChaosOutfitAsset, Warning, TEXT("Piece [%s] already exists in this Outfit!"), *Piece.Name.ToString());
 			AssetGuids.Add(Piece.AssetGuid);
 		}
 	}
