Outfit : make sure to call the property post serialize when loading in collection in an OutFit asset

diff --git a/Engine/Plugins/Experimental/ChaosOutfitAsset/Source/ChaosOutfitAssetEngine/Private/ChaosOutfitAsset/Outfit.cpp b/Engine/Plugins/Experimental/ChaosOutfitAsset/Source/ChaosOutfitAssetEngine/Private/ChaosOutfitAsset/Outfit.cpp
--- a/Engine/Plugins/Experimental/ChaosOutfitAsset/Source/ChaosOutfitAssetEngine/Private/ChaosOutfitAsset/Outfit.cpp
+++ b/Engine/Plugins/Experimental/ChaosOutfitAsset/Source/ChaosOutfitAssetEngine/Private/ChaosOutfitAsset/Outfit.cpp
@@ -622,6 +622,9 @@ bool FChaosOutfitPiece::Serialize(FArchive& Ar)
 			UE::Chaos::ClothAsset::FCollectionClothFacade ClothFacade(NewCollection);
 			ClothFacade.PostSerialize(Ar);
 
+			Chaos::Softs::FCollectionPropertyMutableFacade MutablePropertyFacade(NewCollection);
+			MutablePropertyFacade.PostSerialize(Ar);
+
 			Collections.Emplace(MoveTemp(NewCollection));
 		}
 	}
