Chaos Outfit Asset - Fixed crash when filtering out all render sections with cloth mappings.

[RN] Chaos Outfit Asset - Fixed crash when filtering out all render sections with cloth mappings.

diff --git a/Engine/Plugins/Experimental/ChaosOutfitAsset/Source/ChaosOutfitAssetEngine/Private/ChaosOutfitAsset/Outfit.cpp b/Engine/Plugins/Experimental/ChaosOutfitAsset/Source/ChaosOutfitAssetEngine/Private/ChaosOutfitAsset/Outfit.cpp
--- a/Engine/Plugins/Experimental/ChaosOutfitAsset/Source/ChaosOutfitAssetEngine/Private/ChaosOutfitAsset/Outfit.cpp
+++ b/Engine/Plugins/Experimental/ChaosOutfitAsset/Source/ChaosOutfitAssetEngine/Private/ChaosOutfitAsset/Outfit.cpp
@@ -1100,12 +1100,15 @@ void UChaosOutfit::MergeLODRenderDatas(
 		TArray<FMeshToMeshVertData> ClothMappingData = GetClothMappingData(OutLODRenderData.RenderSections);
 		ClothMappingData.Append(GetClothMappingData(LODRenderData.RenderSections, &AssetGuids));
 
-		TArray<FClothBufferIndexMapping> ClothIndexMappings =
-			GetClothBufferIndexMappings(OutLODRenderData.ClothVertexBuffer, OutLODRenderData.RenderSections);
-		ClothIndexMappings.Append(
-			GetClothBufferIndexMappings(LODRenderData.ClothVertexBuffer, LODRenderData.RenderSections, VertexOffset, &AssetGuids));
+		if (ClothMappingData.Num())  // Make sure the sections that have mappings haven't been filtered since ClothVertexBuffer.Init doesn't take empty mapping without crashing
+		{
+			TArray<FClothBufferIndexMapping> ClothIndexMappings =
+				GetClothBufferIndexMappings(OutLODRenderData.ClothVertexBuffer, OutLODRenderData.RenderSections);
+			ClothIndexMappings.Append(
+				GetClothBufferIndexMappings(LODRenderData.ClothVertexBuffer, LODRenderData.RenderSections, VertexOffset, &AssetGuids));
 
-		OutLODRenderData.ClothVertexBuffer.Init(ClothMappingData, ClothIndexMappings);
+			OutLODRenderData.ClothVertexBuffer.Init(ClothMappingData, ClothIndexMappings);
+		}
 	}
 
 	// Merge sections
