Chaos Outfit Asset - Enabled the Dataflow Editor's evaluations while PIE is running.

#rnx

diff --git a/Engine/Plugins/Experimental/ChaosOutfitAsset/Source/ChaosOutfitAssetEditor/Private/ChaosOutfitAsset/AssetDefinition_OutfitAsset.cpp b/Engine/Plugins/Experimental/ChaosOutfitAsset/Source/ChaosOutfitAssetEditor/Private/ChaosOutfitAsset/AssetDefinition_OutfitAsset.cpp
--- a/Engine/Plugins/Experimental/ChaosOutfitAsset/Source/ChaosOutfitAssetEditor/Private/ChaosOutfitAsset/AssetDefinition_OutfitAsset.cpp
+++ b/Engine/Plugins/Experimental/ChaosOutfitAsset/Source/ChaosOutfitAssetEditor/Private/ChaosOutfitAsset/AssetDefinition_OutfitAsset.cpp
@@ -3,6 +3,8 @@
 #include "ChaosOutfitAsset/AssetDefinition_OutfitAsset.h"
 #include "ChaosOutfitAsset/OutfitAsset.h"
 #include "Dataflow/DataflowEditor.h"
+#include "Dataflow/DataflowEditorToolkit.h"
+#include "Dataflow/DataflowSimulationScene.h"
 #include "ThumbnailRendering/SceneThumbnailInfo.h"
 #include "Toolkits/SimpleAssetEditor.h"
 
@@ -50,12 +52,28 @@ EAssetCommandResult UAssetDefinition_OutfitAsset::OpenAssets(const FAssetOpenArg
 		{
 			// Has a Dataflow Asset, open in the Dataflow editor
 			UAssetEditorSubsystem* const AssetEditorSubsystem = GEditor->GetEditorSubsystem<UAssetEditorSubsystem>();
-			UDataflowEditor* const AssetEditor = NewObject<UDataflowEditor>(AssetEditorSubsystem, NAME_None, RF_Transient);
-			AssetEditor->RegisterToolCategories({"General", "Cloth"});
+			if (UDataflowEditor* const AssetEditor = NewObject<UDataflowEditor>(AssetEditorSubsystem, NAME_None, RF_Transient))
+			{
+				AssetEditor->RegisterToolCategories({"General", "Cloth"});
 
-			static const TCHAR* const Path = TEXT("/ChaosOutfitAsset/BP_OutfitPreview.BP_OutfitPreview_C");
-			const TSubclassOf<AActor> PreviewClass = StaticLoadClass(AActor::StaticClass(), nullptr, Path);
-			AssetEditor->Initialize({ OutfitObjects[0] }, PreviewClass);
+				if (UDataflowSimulationSettings* const SimulationSettings = NewObject<UDataflowSimulationSettings>())
+				{
+					SimulationSettings->bIsSimulationPlayingByDefault = true;
+					SimulationSettings->bIsAsyncCachingSupported = false;
+					SimulationSettings->bIsAsyncCachingEnabledByDefault = false;
+					AssetEditor->AddEditorSettings(SimulationSettings);
+				}
+
+				if (UDataflowEvaluationSettings* const EvaluationSettings = NewObject<UDataflowEvaluationSettings>())
+				{
+					EvaluationSettings->bAllowEvaluationInPIE = true;
+					AssetEditor->AddEditorSettings(EvaluationSettings);
+				}
+
+				static const TCHAR* const Path = TEXT("/ChaosOutfitAsset/BP_OutfitPreview.BP_OutfitPreview_C");
+				const TSubclassOf<AActor> PreviewClass = StaticLoadClass(AActor::StaticClass(), nullptr, Path);
+				AssetEditor->Initialize({ OutfitObjects[0] }, PreviewClass);
+			}
 		}
 		else
 		{
