Chaos Outfit Asset:
- Restructured UChaosOutfitAsset::PostLoad to use the USkinnedAsset async build pipeline (BeginPostLoadInternal/FinishPostLoadInternal) with proper AcquireAsyncProperty/ReleaseAsyncProperty locking, matching UChaosClothAsset.
Without this, GPU resource creation (InitResources) runs outside the async pipeline and any consumer that opens the asset before the deferred FinishPostLoadInternal completes may see an empty mesh.
- Removed duplicate call to CalculateBounds(); that is already made within the Build function and is no longer required since Bounds are now serilalized since 5.6.

#rnx
[RN] Chaos Outfit Asset - Fixed some visibility issues when the asset is first opened in the Dataflow Editor.
#Synced-CL 52356604

diff --git a/Engine/Plugins/Experimental/ChaosOutfitAsset/Source/ChaosOutfitAssetEngine/Private/ChaosOutfitAsset/OutfitAsset.cpp b/Engine/Plugins/Experimental/ChaosOutfitAsset/Source/ChaosOutfitAssetEngine/Private/ChaosOutfitAsset/OutfitAsset.cpp
--- a/Engine/Plugins/Experimental/ChaosOutfitAsset/Source/ChaosOutfitAssetEngine/Private/ChaosOutfitAsset/OutfitAsset.cpp
+++ b/Engine/Plugins/Experimental/ChaosOutfitAsset/Source/ChaosOutfitAssetEngine/Private/ChaosOutfitAsset/OutfitAsset.cpp
@@ -7,6 +7,7 @@
 #include "ChaosOutfitAsset/OutfitAssetPrivate.h"
 #include "Dataflow/DataflowContextAssetStore.h"
 #include "Engine/SkeletalMesh.h"
+#include "Engine/SkinnedAssetAsyncCompileUtils.h"
 #if WITH_EDITOR
 #include "DerivedDataCacheInterface.h"
 #endif
@@ -243,7 +244,23 @@ void UChaosOutfitAsset::Serialize(FArchive& Ar)
 void UChaosOutfitAsset::PostLoad()
 {
 	LLM_SCOPE_BYNAME(TEXT("Physics/Cloth"));
+	// Super::PostLoad drives the BeginPostLoadInternal/FinishPostLoadInternal chain
 	Super::PostLoad();
+}
+
+void UChaosOutfitAsset::BeginPostLoadInternal(FSkinnedAssetPostLoadContext& Context)
+{
+#if WITH_EDITOR
+	TRACE_CPUPROFILER_EVENT_SCOPE(UChaosOutfitAsset::BeginPostLoadInternal);
+
+	checkf(IsInGameThread(), TEXT("Cannot execute function UChaosOutfitAsset::BeginPostLoadInternal asynchronously. Asset: %s"), *GetFullName());
+	SetInternalFlags(EInternalObjectFlags::Async);
+
+	// Lock all properties that should not be modified/accessed during async post-load
+	USkinnedAsset::AcquireAsyncProperty();
+
+	// This scope allows us to use any locked properties without causing stalls
+	FSkinnedAssetAsyncBuildScope AsyncBuildScope(this);
 
 #if WITH_EDITORONLY_DATA
 	if (Outfit)
@@ -252,15 +269,29 @@ void UChaosOutfitAsset::PostLoad()
 		Build(Outfit);
 	}
 	else
+#endif  // #if WITH_EDITORONLY_DATA
+	// Re-evaluate the Dataflow (legacy PostLoad behavior from before the Outfit object was saved)
+	if (GetDataflowInstance().UpdateOwnerAsset())
 	{
-		// Re-evaluate the Dataflow (legacy PostLoad behavior from before the Outfit object was saved)
-		GetDataflowInstance().UpdateOwnerAsset();
 		if (UPackage* const Package = GetOutermost())
 		{
 			Package->SetDirtyFlag(true);
 		}
 		UE_LOG(LogChaosOutfitAsset, Warning, TEXT("Outfit Asset [%s] needs to be re-saved."), *GetName());
 	}
+#endif  // #if WITH_EDITOR
+}
+
+void UChaosOutfitAsset::FinishPostLoadInternal(FSkinnedAssetPostLoadContext& Context)
+{
+#if WITH_EDITOR
+	TRACE_CPUPROFILER_EVENT_SCOPE(UChaosOutfitAsset::FinishPostLoadInternal);
+
+	checkf(IsInGameThread(), TEXT("Cannot execute function UChaosOutfitAsset::FinishPostLoadInternal asynchronously. Asset: %s"), *this->GetFullName());
+	ClearInternalFlags(EInternalObjectFlags::Async);
+
+	// This scope allows us to use any locked properties without causing stalls
+	FSkinnedAssetAsyncBuildScope AsyncBuildScope(this);
 #endif
 
 	if (FApp::CanEverRender())
@@ -273,7 +304,10 @@ void UChaosOutfitAsset::PostLoad()
 	}
 
 	CalculateInvRefMatrices();
-	CalculateBounds();
+
+#if WITH_EDITOR
+	USkinnedAsset::ReleaseAsyncProperty();
+#endif
 }
 
 FName UChaosOutfitAsset::GetClothSimulationModelName(int32 ModelIndex) const
diff --git a/Engine/Plugins/Experimental/ChaosOutfitAsset/Source/ChaosOutfitAssetEngine/Public/ChaosOutfitAsset/OutfitAsset.h b/Engine/Plugins/Experimental/ChaosOutfitAsset/Source/ChaosOutfitAssetEngine/Public/ChaosOutfitAsset/OutfitAsset.h
--- a/Engine/Plugins/Experimental/ChaosOutfitAsset/Source/ChaosOutfitAssetEngine/Public/ChaosOutfitAsset/OutfitAsset.h
+++ b/Engine/Plugins/Experimental/ChaosOutfitAsset/Source/ChaosOutfitAssetEngine/Public/ChaosOutfitAsset/OutfitAsset.h
@@ -84,6 +84,8 @@ private:
 		return nullptr;
 	}
 #endif
+	virtual void BeginPostLoadInternal(FSkinnedAssetPostLoadContext& Context) override;
+	virtual void FinishPostLoadInternal(FSkinnedAssetPostLoadContext& Context) override;
 	//~ End USkinnedAsset interface
 
 	UE_API void CalculateBounds();
