Outfit asset : use the new template picker when creating a new asset

#Synced-CL 52463862

diff --git a/Engine/Plugins/Experimental/ChaosOutfitAsset/Source/ChaosOutfitAssetEditor/Private/ChaosOutfitAsset/ChaosOutfitAssetEditorModule.cpp b/Engine/Plugins/Experimental/ChaosOutfitAsset/Source/ChaosOutfitAssetEditor/Private/ChaosOutfitAsset/ChaosOutfitAssetEditorModule.cpp
--- a/Engine/Plugins/Experimental/ChaosOutfitAsset/Source/ChaosOutfitAssetEditor/Private/ChaosOutfitAsset/ChaosOutfitAssetEditorModule.cpp
+++ b/Engine/Plugins/Experimental/ChaosOutfitAsset/Source/ChaosOutfitAssetEditor/Private/ChaosOutfitAsset/ChaosOutfitAssetEditorModule.cpp
@@ -7,6 +7,7 @@
 #include "ChaosOutfitAsset/OutfitAsset.h"
 #include "ChaosOutfitAsset/OutfitAssetEditorStyle.h"
 #include "ChaosOutfitAsset/OutfitEditorCommands.h"
+#include "Dataflow/DataflowTemplateRegistry.h"
 #include "Engine/SkeletalMesh.h"
 #include "Animation/Skeleton.h"
 #include "Factories/SkeletonFactory.h"
@@ -40,6 +41,44 @@ namespace UE::Chaos::OutfitAsset
 
 			UThumbnailManager::Get().RegisterCustomRenderer(UChaosOutfitAsset::StaticClass(), UChaosClothAssetThumbnailRenderer_Internal::StaticClass());
 
+			const FSlateIcon ChaosOutfitAssetIcon("OutfitAssetEditorStyle", "ClassThumbnail.ChaosOutfitAsset");
+
+			FDataflowTemplateRegistry::Get().RegisterTemplateAsset(
+				UChaosOutfitAsset::StaticClass(),
+				{
+					.AssetPath = FSoftObjectPath(TEXT("/ChaosOutfitAsset/EmptyOutfitAssetTemplate.EmptyOutfitAssetTemplate")),
+					.DisplayName = LOCTEXT("SelectEmptyTemplate", "Empty Dataflow"),
+					.Tooltip = LOCTEXT("SelectEmptyTemplateTooltip", "Add an empty Dataflow with an Outfit Asset Terminal node."),
+					 .Icon = ChaosOutfitAssetIcon
+				});
+
+			FDataflowTemplateRegistry::Get().RegisterTemplateAsset(
+				UChaosOutfitAsset::StaticClass(),
+				{
+					.AssetPath = FSoftObjectPath(TEXT("/ChaosOutfitAsset/OutfitAssetTemplate.OutfitAssetTemplate")),
+					.DisplayName = LOCTEXT("SelectSimpleOutfitTemplate", "Simple Outfit"),
+					.Tooltip = LOCTEXT("SelectSimpleOutfitTemplateTooltip", "Add a Dataflow with a simple Cloth Asset aggregator graph. Allows to simulate multiple Cloth Assets from the same ChaosClothComponent."),
+					 .Icon = ChaosOutfitAssetIcon
+				});
+
+			FDataflowTemplateRegistry::Get().RegisterTemplateAsset(
+				UChaosOutfitAsset::StaticClass(),
+				{
+					.AssetPath = FSoftObjectPath(TEXT("/ChaosOutfitAsset/MakeResizableOutfitTemplate.MakeResizableOutfitTemplate")),
+					.DisplayName = LOCTEXT("SelectResizableOutfitTemplate", "Resizable Outfit"),
+					.Tooltip = LOCTEXT("SelectResizableOutfitTemplateTooltip", "Add a Dataflow that builds a single resizable garment from multiple Cloth Assets. Body sizes will have to be provided in addition to the multiple Cloth Assets."),
+					 .Icon = ChaosOutfitAssetIcon
+				});
+
+			FDataflowTemplateRegistry::Get().RegisterTemplateAsset(
+				UChaosOutfitAsset::StaticClass(),
+				{
+					.AssetPath = FSoftObjectPath(TEXT("/ChaosOutfitAsset/ResizeOutfitTemplate.ResizeOutfitTemplate")),
+					.DisplayName = LOCTEXT("SelectResizingGraph", "Resizing Graph"),
+					.Tooltip = LOCTEXT("SelectResizingGraphTooltip", "Resizing graph to test a resizable Outfit Asset."),
+					 .Icon = ChaosOutfitAssetIcon
+				});
+
 			// TODO Register visualizations (see FChaosClothAssetEditorModule::StartupModule() in ChaosClothAssetEditorModule.cpp)
 		}
 
diff --git a/Engine/Plugins/Experimental/ChaosOutfitAsset/Source/ChaosOutfitAssetEditor/Private/ChaosOutfitAsset/OutfitAssetFactory.cpp b/Engine/Plugins/Experimental/ChaosOutfitAsset/Source/ChaosOutfitAssetEditor/Private/ChaosOutfitAsset/OutfitAssetFactory.cpp
--- a/Engine/Plugins/Experimental/ChaosOutfitAsset/Source/ChaosOutfitAssetEditor/Private/ChaosOutfitAsset/OutfitAssetFactory.cpp
+++ b/Engine/Plugins/Experimental/ChaosOutfitAsset/Source/ChaosOutfitAssetEditor/Private/ChaosOutfitAsset/OutfitAssetFactory.cpp
@@ -12,14 +12,6 @@
 
 DEFINE_LOG_CATEGORY_STATIC(LogChaosOutfitAssetFactory, Log, All);
 
-namespace UE::Chaos::OutfitAsset::Private
-{
-	static TAutoConsoleVariable<bool> CVarCreateWithEmbeddedDataflow(
-		TEXT("p.ChaosOutfitAsset.CreateWithEmbeddedDataflow"),
-		true,
-		TEXT("Set how to create new Outfit Assets, either with embedded (default) or external (legacy) Dataflow assets.\n"));
-}
-
 UChaosOutfitAssetFactory::UChaosOutfitAssetFactory(const FObjectInitializer& ObjectInitializer)
 	: Super(ObjectInitializer)
 {
@@ -30,50 +22,18 @@ UChaosOutfitAssetFactory::UChaosOutfitAssetFactory(const FObjectInitializer& Obj
 
 UObject* UChaosOutfitAssetFactory::FactoryCreateNew(UClass* Class, UObject* Parent, FName Name, EObjectFlags Flags, UObject* /*Context*/, FFeedbackContext* /*Warn*/)
 {
-	using namespace UE::Chaos::OutfitAsset;
-
-	// Select the template
-	TSharedRef<SMessageDialog> SelectTemplateMessageDialog = SNew(SMessageDialog)
-		.Title(FText(LOCTEXT("SelectTemplateTitle", "Select an Outfit Asset Template")))
-		.Message(LOCTEXT("SelectTemplateMessage", "Select a template for this newly created Outfit Asset:"))
-		.Buttons(
-			{
-				SMessageDialog::FButton(LOCTEXT("SelectNoTemplate", "No Dataflow"))
-					.SetToolTipText(LOCTEXT("SelectNoTemplateTooltip", "Don't add a Dataflow to this Outfit Asset. Useful when a Dataflow already exists."))
-					.SetPrimary(false),
-				SMessageDialog::FButton(LOCTEXT("SelectEmptyTemplate", "Empty Dataflow"))
-					.SetToolTipText(LOCTEXT("SelectEmptyTemplateTooltip", "Add an empty Dataflow with an Outfit Asset Terminal node."))
-					.SetPrimary(false),
-				SMessageDialog::FButton(LOCTEXT("SelectSimpleOutfitTemplate", "Simple Outfit"))
-					.SetToolTipText(LOCTEXT("SelectSimpleOutfitTemplateTooltip", "Add a Dataflow with a simple Cloth Asset aggregator graph. Allows to simulate multiple Cloth Assets from the same ChaosClothComponent."))
-					.SetPrimary(true),
-				SMessageDialog::FButton(LOCTEXT("SelectResizableOutfitTemplate", "Resizable Outfit"))
-					.SetToolTipText(LOCTEXT("SelectResizableOutfitTemplateTooltip", "Add a Dataflow that builds a single resizable garment from multiple Cloth Assets. Body sizes will have to be provided in addition to the multiple Cloth Assets."))
-					.SetPrimary(false),
-				SMessageDialog::FButton(LOCTEXT("SelectResizingGraph", "Resizing Graph"))
-					.SetToolTipText(LOCTEXT("SelectResizingGraphTooltip", "Resizing graph to test a resizable Outfit Asset."))
-					.SetPrimary(false),
-			});
-
-	const FString* TemplatePath;
-	static const FString EmptyOutfitTemplatePath(TEXT("/ChaosOutfitAsset/EmptyOutfitAssetTemplate.EmptyOutfitAssetTemplate"));
-	static const FString SimpleOutfitTemplatePath(TEXT("/ChaosOutfitAsset/OutfitAssetTemplate.OutfitAssetTemplate"));
-	static const FString ResizableOutfitTemplatePath(TEXT("/ChaosOutfitAsset/MakeResizableOutfitTemplate.MakeResizableOutfitTemplate"));
-	static const FString ResizingOutfitTemplatePath(TEXT("/ChaosOutfitAsset/ResizeOutfitTemplate.ResizeOutfitTemplate"));
-	switch (SelectTemplateMessageDialog->ShowModal())
+	UObject* const NewAsset = NewObject<UObject>(Parent, Class, Name, Flags | RF_Transactional | RF_Public | RF_Standalone);
+	if (!NewAsset)
 	{
-	default:
-	case 0: TemplatePath = nullptr; break;  // No Dataflow
-	case 1: TemplatePath = &EmptyOutfitTemplatePath; break;  // Empty Dataflow with terminal node
-	case 2: TemplatePath = &SimpleOutfitTemplatePath; break;  // Simple Outfit
-	case 3: TemplatePath = &ResizableOutfitTemplatePath; break;  // Resizable Outfit
-	case 4: TemplatePath = &ResizingOutfitTemplatePath; break;  // Resizing Graph
+		const FText ErrorMessage = FText::Format(LOCTEXT("ErrorCreatingNewAsset", "Failed to create a new asset {0}."), FText::FromName(Name));
+		UE_LOG(LogChaosOutfitAssetFactory, Error, TEXT("%s"), *ErrorMessage.ToString());
+		return nullptr;
 	}
+	NewAsset->MarkPackageDirty();
 
-	// Create a new asset using the provided template path
-	static const TCHAR* AssetPrefix = TEXT("OA_");
-	const bool bEmbedDataflow = Private::CVarCreateWithEmbeddedDataflow.GetValueOnAnyThread();
-	return UE::DataflowAssetDefinitionHelpers::FactoryCreateNew(Class, Parent, Name, Flags, TemplatePath, bEmbedDataflow, AssetPrefix);
+	// No default blank option since cloth templates already provide one
+	UE::DataflowAssetDefinitionHelpers::SetDataflowFromTemplatePicker(NewAsset, /*bShowDefaultBlankOption*/ false);
+	return NewAsset;
 }
 
 FString UChaosOutfitAssetFactory::GetDefaultNewAssetName() const
